Fix include warnings issue

diff --git a/Engine/Source/Programs/AutoRTFMTests/Private/AutoRTFMTests.cpp b/Engine/Source/Programs/AutoRTFMTests/Private/AutoRTFMTests.cpp
--- a/Engine/Source/Programs/AutoRTFMTests/Private/AutoRTFMTests.cpp
+++ b/Engine/Source/Programs/AutoRTFMTests/Private/AutoRTFMTests.cpp
@@ -1,8 +1,7 @@
 // Copyright Epic Games, Inc. All Rights Reserved.
 
-#include "Modules/ModuleManager.h"
-#include "LaunchEngineLoop.h"   // GEngineLoop
-
+#include "CoreMinimal.h"
+#include "CoreGlobals.h"
 #include "RequiredProgramMainCPPInclude.h" // required for ue programs
 
 IMPLEMENT_APPLICATION(AutoRTFMTests, "AutoRTFMTests");
