Dont enable ClangRTFM when doing GPF

diff --git a/Engine/Source/Programs/AutoRTFMTests/AutoRTFMTests.Target.cs b/Engine/Source/Programs/AutoRTFMTests/AutoRTFMTests.Target.cs
--- a/Engine/Source/Programs/AutoRTFMTests/AutoRTFMTests.Target.cs
+++ b/Engine/Source/Programs/AutoRTFMTests/AutoRTFMTests.Target.cs
@@ -24,6 +24,9 @@ public class AutoRTFMTestsTarget : TargetRules
 		bIsBuildingConsoleApplication = true;
 
 		// Set the RTFM clang compiler
-		WindowsPlatform.Compiler = WindowsCompiler.ClangRTFM;
+		if(!bGenerateProjectFiles)
+		{
+		     WindowsPlatform.Compiler = WindowsCompiler.ClangRTFM
+		}
 	}
 }
