Fix missing semi-colon

diff --git a/Engine/Source/Programs/AutoRTFMTests/AutoRTFMTests.Target.cs b/Engine/Source/Programs/AutoRTFMTests/AutoRTFMTests.Target.cs
--- a/Engine/Source/Programs/AutoRTFMTests/AutoRTFMTests.Target.cs
+++ b/Engine/Source/Programs/AutoRTFMTests/AutoRTFMTests.Target.cs
@@ -26,7 +26,7 @@ public class AutoRTFMTestsTarget : TargetRules
 		// Set the RTFM clang compiler
 		if(!bGenerateProjectFiles)
 		{
-		     WindowsPlatform.Compiler = WindowsCompiler.ClangRTFM
+		     WindowsPlatform.Compiler = WindowsCompiler.ClangRTFM;
 		}
 	}
 }
