Add a header to ensure we include our MinWindows and Windows type headers *before* catch_amalgamated.hpp includes its Windows.h stuff
Fix GraphBenchmark shadow warnings

diff --git a/Engine/Source/Programs/AutoRTFMTests/Private/API.cpp b/Engine/Source/Programs/AutoRTFMTests/Private/API.cpp
--- a/Engine/Source/Programs/AutoRTFMTests/Private/API.cpp
+++ b/Engine/Source/Programs/AutoRTFMTests/Private/API.cpp
@@ -3,8 +3,8 @@
 #pragma autortfm
 
 #include "API.h"
+#include "Catch2Includes.h"
 
-#include <catch_amalgamated.hpp>
 #include <AutoRTFM/AutoRTFM.h>
 #include <memory>
 #include <thread>
diff --git a/Engine/Source/Programs/AutoRTFMTests/Private/Abort.cpp b/Engine/Source/Programs/AutoRTFMTests/Private/Abort.cpp
--- a/Engine/Source/Programs/AutoRTFMTests/Private/Abort.cpp
+++ b/Engine/Source/Programs/AutoRTFMTests/Private/Abort.cpp
@@ -2,7 +2,7 @@
 
 #pragma autortfm
 
-#include <catch_amalgamated.hpp>
+#include "Catch2Includes.h"
 #include <AutoRTFM/AutoRTFM.h>
 #include <map>
 #include <vector>
diff --git a/Engine/Source/Programs/AutoRTFMTests/Private/AutoRTFMTests.cpp b/Engine/Source/Programs/AutoRTFMTests/Private/AutoRTFMTests.cpp
--- a/Engine/Source/Programs/AutoRTFMTests/Private/AutoRTFMTests.cpp
+++ b/Engine/Source/Programs/AutoRTFMTests/Private/AutoRTFMTests.cpp
@@ -1,7 +1,6 @@
 // Copyright Epic Games, Inc. All Rights Reserved.
 
 #include "CoreMinimal.h"
-#include "CoreGlobals.h"
 #include "RequiredProgramMainCPPInclude.h" // required for ue programs
 
 IMPLEMENT_APPLICATION(AutoRTFMTests, "AutoRTFMTests");
diff --git a/Engine/Source/Programs/AutoRTFMTests/Private/Benchmarks.cpp b/Engine/Source/Programs/AutoRTFMTests/Private/Benchmarks.cpp
--- a/Engine/Source/Programs/AutoRTFMTests/Private/Benchmarks.cpp
+++ b/Engine/Source/Programs/AutoRTFMTests/Private/Benchmarks.cpp
@@ -2,7 +2,7 @@
 
 #pragma autortfm
 
-#include <catch_amalgamated.hpp>
+#include "Catch2Includes.h"
 #include <AutoRTFM/AutoRTFM.h>
 #include <thread>
 #include <vector>
diff --git a/Engine/Source/Programs/AutoRTFMTests/Private/Catch2Includes.h b/Engine/Source/Programs/AutoRTFMTests/Private/Catch2Includes.h
new file mode 100644
--- /dev/null
+++ b/Engine/Source/Programs/AutoRTFMTests/Private/Catch2Includes.h
@@ -0,0 +1,7 @@
+// Copyright Epic Games, Inc. All Rights Reserved.
+
+#pragma once
+
+// we need to include thie MainCPP include *before* we include catch_amalgamated or we get issues with windows symbols
+#include "RequiredProgramMainCPPInclude.h"
+#include <catch_amalgamated.hpp>
diff --git a/Engine/Source/Programs/AutoRTFMTests/Private/Commit.cpp b/Engine/Source/Programs/AutoRTFMTests/Private/Commit.cpp
--- a/Engine/Source/Programs/AutoRTFMTests/Private/Commit.cpp
+++ b/Engine/Source/Programs/AutoRTFMTests/Private/Commit.cpp
@@ -2,7 +2,7 @@
 
 #pragma autortfm
 
-#include <catch_amalgamated.hpp>
+#include "Catch2Includes.h"
 #include <AutoRTFM/AutoRTFM.h>
 #include <map>
 #include <vector>
diff --git a/Engine/Source/Programs/AutoRTFMTests/Private/CopyAddInt64.cpp b/Engine/Source/Programs/AutoRTFMTests/Private/CopyAddInt64.cpp
--- a/Engine/Source/Programs/AutoRTFMTests/Private/CopyAddInt64.cpp
+++ b/Engine/Source/Programs/AutoRTFMTests/Private/CopyAddInt64.cpp
@@ -2,7 +2,7 @@
 
 #pragma autortfm
 
-#include <catch_amalgamated.hpp>
+#include "Catch2Includes.h"
 #include <AutoRTFM/AutoRTFM.h>
 
 TEST_CASE("CopyAddInt64")
diff --git a/Engine/Source/Programs/AutoRTFMTests/Private/DeferUntilCommit.cpp b/Engine/Source/Programs/AutoRTFMTests/Private/DeferUntilCommit.cpp
--- a/Engine/Source/Programs/AutoRTFMTests/Private/DeferUntilCommit.cpp
+++ b/Engine/Source/Programs/AutoRTFMTests/Private/DeferUntilCommit.cpp
@@ -2,7 +2,7 @@
 
 #pragma autortfm
 
-#include <catch_amalgamated.hpp>
+#include "Catch2Includes.h"
 #include <AutoRTFM/AutoRTFM.h>
 
 TEST_CASE("DeferUntilCommit")
diff --git a/Engine/Source/Programs/AutoRTFMTests/Private/GraphBenchmark.cpp b/Engine/Source/Programs/AutoRTFMTests/Private/GraphBenchmark.cpp
--- a/Engine/Source/Programs/AutoRTFMTests/Private/GraphBenchmark.cpp
+++ b/Engine/Source/Programs/AutoRTFMTests/Private/GraphBenchmark.cpp
@@ -2,7 +2,7 @@
 
 #pragma autortfm
 
-#include <catch_amalgamated.hpp>
+#include "Catch2Includes.h"
 #include <AutoRTFM/AutoRTFM.h>
 #include <vector>
 #include <unordered_set>
@@ -95,9 +95,9 @@ public:
             }
         };
 
-        auto PushAll = [&] (const vector<FNode<T>*>& Nodes)
+        auto PushAll = [&] (const vector<FNode<T>*>& InnerNodes)
         {
-            for (FNode<T>* Node : Nodes)
+            for (FNode<T>* Node : InnerNodes)
             {
                 Push(Node);
             }
diff --git a/Engine/Source/Programs/AutoRTFMTests/Private/Open.cpp b/Engine/Source/Programs/AutoRTFMTests/Private/Open.cpp
--- a/Engine/Source/Programs/AutoRTFMTests/Private/Open.cpp
+++ b/Engine/Source/Programs/AutoRTFMTests/Private/Open.cpp
@@ -2,7 +2,7 @@
 
 #pragma autortfm
 
-#include <catch_amalgamated.hpp>
+#include "Catch2Includes.h"
 #include <AutoRTFM/AutoRTFM.h>
 #include <map>
 #include <vector>
diff --git a/Engine/Source/Programs/AutoRTFMTests/Private/PackedPointer.cpp b/Engine/Source/Programs/AutoRTFMTests/Private/PackedPointer.cpp
--- a/Engine/Source/Programs/AutoRTFMTests/Private/PackedPointer.cpp
+++ b/Engine/Source/Programs/AutoRTFMTests/Private/PackedPointer.cpp
@@ -2,7 +2,7 @@
 
 #pragma autortfm
 
-#include <catch_amalgamated.hpp>
+#include "Catch2Includes.h"
 #include <AutoRTFM/AutoRTFM.h>
 #include <inttypes.h>
 
diff --git a/Engine/Source/Programs/AutoRTFMTests/Private/Tests.cpp b/Engine/Source/Programs/AutoRTFMTests/Private/Tests.cpp
--- a/Engine/Source/Programs/AutoRTFMTests/Private/Tests.cpp
+++ b/Engine/Source/Programs/AutoRTFMTests/Private/Tests.cpp
@@ -2,7 +2,7 @@
 
 #pragma autortfm
 
-#include <catch_amalgamated.hpp>
+#include "Catch2Includes.h"
 #include <AutoRTFM/AutoRTFM.h>
 #include <thread>
 #include <atomic>
diff --git a/Engine/Source/Programs/AutoRTFMTests/Private/UnorderedSet.cpp b/Engine/Source/Programs/AutoRTFMTests/Private/UnorderedSet.cpp
--- a/Engine/Source/Programs/AutoRTFMTests/Private/UnorderedSet.cpp
+++ b/Engine/Source/Programs/AutoRTFMTests/Private/UnorderedSet.cpp
@@ -2,7 +2,7 @@
 
 #pragma autortfm
 
-#include <catch_amalgamated.hpp>
+#include "Catch2Includes.h"
 #include <AutoRTFM/AutoRTFM.h>
 #include <unordered_set>
 
