Disable Unity by default for AutoRTFMTests

diff --git a/Engine/Source/Programs/AutoRTFMTests/AutoRTFMTests.Target.cs b/Engine/Source/Programs/AutoRTFMTests/AutoRTFMTests.Target.cs
--- a/Engine/Source/Programs/AutoRTFMTests/AutoRTFMTests.Target.cs
+++ b/Engine/Source/Programs/AutoRTFMTests/AutoRTFMTests.Target.cs
@@ -23,6 +23,9 @@ public class AutoRTFMTestsTarget : TargetRules
 		// Make a console application under Windows, so entry point is main() everywhere
 		bIsBuildingConsoleApplication = true;
 
+		// Disable unity builds by default for AutoRTFMTests
+		bUseUnityBuild = false;
+
 		// Set the RTFM clang compiler
 		if(!bGenerateProjectFiles)
 		{
