Re-enable unity to fix linker errors

diff --git a/Engine/Source/Programs/AutoRTFMTests/AutoRTFMTests.Target.cs b/Engine/Source/Programs/AutoRTFMTests/AutoRTFMTests.Target.cs
--- a/Engine/Source/Programs/AutoRTFMTests/AutoRTFMTests.Target.cs
+++ b/Engine/Source/Programs/AutoRTFMTests/AutoRTFMTests.Target.cs
@@ -24,10 +24,11 @@ public class AutoRTFMTestsTarget : TargetRules
 		bIsBuildingConsoleApplication = true;
 
 		// Disable unity builds by default for AutoRTFMTests
-		bUseUnityBuild = false;
+		// TODO enabling this causes linker errors
+		// bUseUnityBuild = false;
 
 		// Set the RTFM clang compiler
-		if(!bGenerateProjectFiles)
+		if (!bGenerateProjectFiles)
 		{
 		     WindowsPlatform.Compiler = WindowsCompiler.ClangRTFM;
 		}
