Re-enable disabling unity and remove including LaunchEngineLoop.cpp many times

diff --git a/Engine/Source/Programs/AutoRTFMTests/AutoRTFMTests.Target.cs b/Engine/Source/Programs/AutoRTFMTests/AutoRTFMTests.Target.cs
--- a/Engine/Source/Programs/AutoRTFMTests/AutoRTFMTests.Target.cs
+++ b/Engine/Source/Programs/AutoRTFMTests/AutoRTFMTests.Target.cs
@@ -11,6 +11,8 @@ public class AutoRTFMTestsTarget : TargetRules
 		Type = TargetType.Program;
 		LinkType = TargetLinkType.Monolithic;
 		IncludeOrderVersion = EngineIncludeOrderVersion.Latest;
+
+		Name = "AutoRTFMTests";
 		LaunchModuleName = "AutoRTFMTests";
 
 		// Compile out references from Core to the rest of the engine
@@ -23,9 +25,8 @@ public class AutoRTFMTestsTarget : TargetRules
 		// Make a console application under Windows, so entry point is main() everywhere
 		bIsBuildingConsoleApplication = true;
 
-		// Disable unity builds by default for AutoRTFMTests
-		// TODO enabling this causes linker errors
-		// bUseUnityBuild = false;
+		// Disable unity builds by default for AutoRTFMTest
+		bUseUnityBuild = false;
 
 		// Set the RTFM clang compiler
 		if (!bGenerateProjectFiles)
diff --git a/Engine/Source/Programs/AutoRTFMTests/Private/Catch2Includes.h b/Engine/Source/Programs/AutoRTFMTests/Private/Catch2Includes.h
--- a/Engine/Source/Programs/AutoRTFMTests/Private/Catch2Includes.h
+++ b/Engine/Source/Programs/AutoRTFMTests/Private/Catch2Includes.h
@@ -3,5 +3,6 @@
 #pragma once
 
 // we need to include thie MainCPP include *before* we include catch_amalgamated or we get issues with windows symbols
-#include "RequiredProgramMainCPPInclude.h"
+#include "CoreMinimal.h"
+#include "LaunchEngineLoop.h"
 #include <catch_amalgamated.hpp>
