Harden the AutoRTFMification LLVM pass against a few obvious compiler bugs. The removed files were all dead as they've been moved to the public area in a previous CL.

diff --git a/Engine/Source/Programs/AutoRTFMTests/Private/Allocators.cpp b/Engine/Source/Programs/AutoRTFMTests/Private/Allocators.cpp
new file mode 100644
--- /dev/null
+++ b/Engine/Source/Programs/AutoRTFMTests/Private/Allocators.cpp
@@ -0,0 +1,83 @@
+// Copyright Epic Games, Inc. All Rights Reserved.
+
+#pragma autortfm
+
+#include <catch_amalgamated.hpp>
+#include <AutoRTFM/AutoRTFM.h>
+#include <new>
+
+TEST_CASE("Allocators.New")
+{
+	int* Data;
+	AutoRTFM::Commit([&]
+		{
+			Data = new int(42);
+		});
+
+	REQUIRE(42 == *Data);
+
+	delete Data;
+}
+
+TEST_CASE("Allocators.Delete")
+{
+	int* Data = new int(42);
+
+	AutoRTFM::Commit([&]
+		{
+			delete Data;
+		});
+}
+
+TEST_CASE("Allocators.ArrayNew")
+{
+	int* Data;
+	AutoRTFM::Commit([&]
+		{
+			Data = new int[42]();
+			Data[2] = 42;
+		});
+
+	REQUIRE(42 == Data[2]);
+
+	delete[] Data;
+}
+
+// TODO: Re-enable test once array delete[] works in transactions.
+TEST_CASE("Allocators.ArrayDelete", "[.]")
+{
+	int* Data = new int[42]();
+
+	AutoRTFM::Commit([&]
+		{
+			delete[] Data;
+		});
+}
+
+TEST_CASE("Allocators.NewNoOpts")
+{
+	int* Data;
+
+#pragma clang optimize off
+	AutoRTFM::Commit([&]
+		{
+			Data = new int(42);
+		});
+#pragma clang optimize on
+
+	REQUIRE(42 == *Data);
+
+	delete Data;
+}
+
+TEST_CASE("Allocators.DeleteNoOpts")
+{
+	int* Data = new int(42);
+
+#pragma clang optimize off
+	AutoRTFM::Commit([&]
+		{
+			delete Data;
+		});
+#pragma clang optimize on
+}
