First go at the AutoRTFM Open API for the VM. This is definitely crunchy.
Adds StartTransaction, RecordOpenWrite, WriteMemory, and some others.
Replaces Abort with AbortTransaction

diff --git a/Engine/Source/Programs/AutoRTFMTests/Private/API.cpp b/Engine/Source/Programs/AutoRTFMTests/Private/API.cpp
--- a/Engine/Source/Programs/AutoRTFMTests/Private/API.cpp
+++ b/Engine/Source/Programs/AutoRTFMTests/Private/API.cpp
@@ -23,7 +23,7 @@ TEST_CASE("API.autortfm_result")
     REQUIRE(autortfm_aborted_by_request == autortfm_transact([](void* const Arg)
     {
         *static_cast<int* const>(Arg) = 13;
-        AutoRTFM::Abort();
+        AutoRTFM::AbortTransaction();
     }, &Answer));
 
     REQUIRE(42 == Answer);
@@ -127,7 +127,7 @@ TEST_CASE("API.autortfm_abort")
             // Because we are aborting this won't actually occur!
             InNest = true;
 
-            autortfm_abort();
+            autortfm_abort_transaction();
         });
 
         AfterNest = true;
@@ -221,7 +221,7 @@ TEST_CASE("API.autortfm_open")
             *static_cast<int* const>(Arg) *= 2;
         }, &Answer);
 
-        AutoRTFM::Abort();
+        AutoRTFM::AbortTransaction();
     }));
 
     REQUIRE(84 == Answer);
@@ -252,7 +252,7 @@ TEST_CASE("API.autortfm_close")
             InOpenNest = true;
         });
 
-        AutoRTFM::Abort();
+        AutoRTFM::AbortTransaction();
     }));
 
     REQUIRE(false == InClosedNest);
@@ -294,7 +294,7 @@ TEST_CASE("API.autortfm_defer_until_commit")
         // This should only be modified on the commit!
         if (OuterTransaction)
         {
-            AutoRTFM::Abort();
+            AutoRTFM::AbortTransaction();
         }
 
         AutoRTFM::Commit([&]
@@ -308,7 +308,7 @@ TEST_CASE("API.autortfm_defer_until_commit")
         // This should only be modified on the commit!
         if (InnerTransaction)
         {
-            AutoRTFM::Abort();
+            AutoRTFM::AbortTransaction();
         }
 
         NestResult = AutoRTFM::Transact([&]
@@ -318,13 +318,13 @@ TEST_CASE("API.autortfm_defer_until_commit")
                 *static_cast<bool* const>(Arg) = true;
             }, &InnerTransactionWithAbort);
 
-            AutoRTFM::Abort();
+            AutoRTFM::AbortTransaction();
         });
 
         // This should never be modified because its transaction aborted!
         if (InnerTransactionWithAbort)
         {
-            AutoRTFM::Abort();
+            AutoRTFM::AbortTransaction();
         }
 
         AutoRTFM::Open([&]
@@ -338,7 +338,7 @@ TEST_CASE("API.autortfm_defer_until_commit")
         // This should only be modified on the commit!
         if (InnerOpenNest)
         {
-            AutoRTFM::Abort();
+            AutoRTFM::AbortTransaction();
         }
     }));
 
@@ -367,7 +367,7 @@ TEST_CASE("API.autortfm_defer_until_abort")
         // This should only be modified on the commit!
         if (OuterTransaction)
         {
-            AutoRTFM::Abort();
+            AutoRTFM::AbortTransaction();
         }
 
         AutoRTFM::Commit([&]
@@ -381,7 +381,7 @@ TEST_CASE("API.autortfm_defer_until_abort")
         // This should only be modified on the commit!
         if (InnerTransaction)
         {
-            AutoRTFM::Abort();
+            AutoRTFM::AbortTransaction();
         }
 
         NestResult = AutoRTFM::Transact([&]
@@ -391,13 +391,13 @@ TEST_CASE("API.autortfm_defer_until_abort")
                 *static_cast<bool* const>(Arg) = true;
             }, &InnerTransactionWithAbort);
 
-            AutoRTFM::Abort();
+            AutoRTFM::AbortTransaction();
         });
 
         // This should never be modified because its transaction aborted!
         if (InnerTransactionWithAbort)
         {
-            AutoRTFM::Abort();
+            AutoRTFM::AbortTransaction();
         }
 
         AutoRTFM::Open([&]
@@ -411,7 +411,7 @@ TEST_CASE("API.autortfm_defer_until_abort")
         // This should only be modified on the commit!
         if (InnerOpenNest)
         {
-            AutoRTFM::Abort();
+            AutoRTFM::AbortTransaction();
         }
     }));
 
@@ -440,7 +440,7 @@ TEST_CASE("API.autortfm_open_commit")
         // This should only be modified on the commit!
         if (OuterTransaction)
         {
-            AutoRTFM::Abort();
+            AutoRTFM::AbortTransaction();
         }
 
         AutoRTFM::Commit([&]
@@ -454,7 +454,7 @@ TEST_CASE("API.autortfm_open_commit")
         // This should only be modified on the commit!
         if (InnerTransaction)
         {
-            AutoRTFM::Abort();
+            AutoRTFM::AbortTransaction();
         }
 
         NestResult = AutoRTFM::Transact([&]
@@ -464,13 +464,13 @@ TEST_CASE("API.autortfm_open_commit")
                 *static_cast<bool* const>(Arg) = true;
             }, &InnerTransactionWithAbort);
 
-            AutoRTFM::Abort();
+            AutoRTFM::AbortTransaction();
         });
 
         // This should never be modified because its transaction aborted!
         if (InnerTransactionWithAbort)
         {
-            AutoRTFM::Abort();
+            AutoRTFM::AbortTransaction();
         }
 
         AutoRTFM::Open([&]
@@ -483,7 +483,7 @@ TEST_CASE("API.autortfm_open_commit")
             // This should be modified immediately!
             if (!InnerOpenNest)
             {
-                AutoRTFM::Abort();
+                AutoRTFM::AbortTransaction();
             }
         });
     });
@@ -513,7 +513,7 @@ TEST_CASE("API.autortfm_open_abort")
         // This should only be modified on the commit!
         if (OuterTransaction)
         {
-            AutoRTFM::Abort();
+            AutoRTFM::AbortTransaction();
         }
 
         AutoRTFM::Commit([&]
@@ -527,7 +527,7 @@ TEST_CASE("API.autortfm_open_abort")
         // This should only be modified on the commit!
         if (InnerTransaction)
         {
-            AutoRTFM::Abort();
+            AutoRTFM::AbortTransaction();
         }
 
         NestResult = AutoRTFM::Transact([&]
@@ -537,13 +537,13 @@ TEST_CASE("API.autortfm_open_abort")
                 *static_cast<bool* const>(Arg) = true;
             }, &InnerTransactionWithAbort);
 
-            AutoRTFM::Abort();
+            AutoRTFM::AbortTransaction();
         });
 
         // This should never be modified because its transaction aborted!
         if (InnerTransactionWithAbort)
         {
-            AutoRTFM::Abort();
+            AutoRTFM::AbortTransaction();
         }
 
         AutoRTFM::Open([&]
@@ -557,7 +557,7 @@ TEST_CASE("API.autortfm_open_abort")
         // This should only be modified on the commit!
         if (InnerOpenNest)
         {
-            AutoRTFM::Abort();
+            AutoRTFM::AbortTransaction();
         }
     }));
 
@@ -645,7 +645,7 @@ TEST_CASE("API.ETransactionResult")
     REQUIRE(AutoRTFM::ETransactionResult::AbortedByRequest == AutoRTFM::Transact([&]
     {
         Answer = 13;
-        AutoRTFM::Abort();
+        AutoRTFM::AbortTransaction();
     }));
 
     REQUIRE(42 == Answer);
@@ -749,7 +749,7 @@ TEST_CASE("API.Abort")
             // Because we are aborting this won't actually occur!
             InNest = true;
 
-            AutoRTFM::Abort();
+            AutoRTFM::AbortTransaction();
         });
 
         AfterNest = true;
@@ -843,7 +843,7 @@ TEST_CASE("API.Open")
             Answer *= 2;
         });
 
-        AutoRTFM::Abort();
+        AutoRTFM::AbortTransaction();
     }));
 
     REQUIRE(84 == Answer);
@@ -874,7 +874,7 @@ TEST_CASE("API.Close")
             InOpenNest = true;
         });
 
-        AutoRTFM::Abort();
+        AutoRTFM::AbortTransaction();
     }));
 
     REQUIRE(false == InClosedNest);
@@ -916,7 +916,7 @@ TEST_CASE("API.DeferUntilCommit")
         // This should only be modified on the commit!
         if (OuterTransaction)
         {
-            AutoRTFM::Abort();
+            AutoRTFM::AbortTransaction();
         }
 
         AutoRTFM::Commit([&]
@@ -930,7 +930,7 @@ TEST_CASE("API.DeferUntilCommit")
         // This should only be modified on the commit!
         if (InnerTransaction)
         {
-            AutoRTFM::Abort();
+            AutoRTFM::AbortTransaction();
         }
 
         NestResult = AutoRTFM::Transact([&]
@@ -940,13 +940,13 @@ TEST_CASE("API.DeferUntilCommit")
                 InnerTransactionWithAbort = true;
             });
 
-            AutoRTFM::Abort();
+            AutoRTFM::AbortTransaction();
         });
 
         // This should never be modified because its transaction aborted!
         if (InnerTransactionWithAbort)
         {
-            AutoRTFM::Abort();
+            AutoRTFM::AbortTransaction();
         }
 
         AutoRTFM::Open([&]
@@ -960,7 +960,7 @@ TEST_CASE("API.DeferUntilCommit")
         // This should only be modified on the commit!
         if (InnerOpenNest)
         {
-            AutoRTFM::Abort();
+            AutoRTFM::AbortTransaction();
         }
     }));
 
@@ -989,7 +989,7 @@ TEST_CASE("API.DeferUntilAbort")
         // This should only be modified on the commit!
         if (OuterTransaction)
         {
-            AutoRTFM::Abort();
+            AutoRTFM::AbortTransaction();
         }
 
         AutoRTFM::Commit([&]
@@ -1003,7 +1003,7 @@ TEST_CASE("API.DeferUntilAbort")
         // This should only be modified on the commit!
         if (InnerTransaction)
         {
-            AutoRTFM::Abort();
+			AutoRTFM::AbortTransaction();
         }
 
         NestResult = AutoRTFM::Transact([&]
@@ -1013,13 +1013,13 @@ TEST_CASE("API.DeferUntilAbort")
                 InnerTransactionWithAbort = true;
             });
 
-            AutoRTFM::Abort();
+            AutoRTFM::AbortTransaction();
         });
 
         // This should never be modified because its transaction aborted!
         if (InnerTransactionWithAbort)
         {
-            AutoRTFM::Abort();
+            AutoRTFM::AbortTransaction();
         }
 
         AutoRTFM::Open([&]
@@ -1033,7 +1033,7 @@ TEST_CASE("API.DeferUntilAbort")
         // This should only be modified on the commit!
         if (InnerOpenNest)
         {
-            AutoRTFM::Abort();
+			AutoRTFM::AbortTransaction();
         }
     }));
 
@@ -1062,7 +1062,7 @@ TEST_CASE("API.OpenCommit")
         // This should only be modified on the commit!
         if (OuterTransaction)
         {
-            AutoRTFM::Abort();
+            AutoRTFM::AbortTransaction();
         }
 
         AutoRTFM::Commit([&]
@@ -1076,7 +1076,7 @@ TEST_CASE("API.OpenCommit")
         // This should only be modified on the commit!
         if (InnerTransaction)
         {
-            AutoRTFM::Abort();
+			AutoRTFM::AbortTransaction();
         }
 
         NestResult = AutoRTFM::Transact([&]
@@ -1086,13 +1086,13 @@ TEST_CASE("API.OpenCommit")
                 InnerTransactionWithAbort = true;
             });
 
-            AutoRTFM::Abort();
+		AutoRTFM::AbortTransaction();
         });
 
         // This should never be modified because its transaction aborted!
         if (InnerTransactionWithAbort)
         {
-            AutoRTFM::Abort();
+			AutoRTFM::AbortTransaction();
         }
 
         AutoRTFM::Open([&]
@@ -1105,7 +1105,7 @@ TEST_CASE("API.OpenCommit")
             // This should be modified immediately!
             if (!InnerOpenNest)
             {
-                AutoRTFM::Abort();
+				AutoRTFM::AbortTransaction();
             }
         });
     }));
@@ -1135,7 +1135,7 @@ TEST_CASE("API.OpenAbort")
         // This should only be modified on the commit!
         if (OuterTransaction)
         {
-            AutoRTFM::Abort();
+			AutoRTFM::AbortTransaction();
         }
 
         AutoRTFM::Commit([&]
@@ -1149,7 +1149,7 @@ TEST_CASE("API.OpenAbort")
         // This should only be modified on the commit!
         if (InnerTransaction)
         {
-            AutoRTFM::Abort();
+			AutoRTFM::AbortTransaction();
         }
 
         NestResult = AutoRTFM::Transact([&]
@@ -1159,13 +1159,13 @@ TEST_CASE("API.OpenAbort")
                 InnerTransactionWithAbort = true;
             });
 
-            AutoRTFM::Abort();
+		AutoRTFM::AbortTransaction();
         });
 
         // This should never be modified because its transaction aborted!
         if (InnerTransactionWithAbort)
         {
-            AutoRTFM::Abort();
+			AutoRTFM::AbortTransaction();
         }
 
         AutoRTFM::Open([&]
@@ -1179,7 +1179,7 @@ TEST_CASE("API.OpenAbort")
         // This should only be modified on the commit!
         if (InnerOpenNest)
         {
-            AutoRTFM::Abort();
+			AutoRTFM::AbortTransaction();
         }
     }));
 
diff --git a/Engine/Source/Programs/AutoRTFMTests/Private/Abort.cpp b/Engine/Source/Programs/AutoRTFMTests/Private/Abort.cpp
--- a/Engine/Source/Programs/AutoRTFMTests/Private/Abort.cpp
+++ b/Engine/Source/Programs/AutoRTFMTests/Private/Abort.cpp
@@ -29,7 +29,7 @@ TEST_CASE("Abort")
     	m[10].push_back(11);
     	m[12].push_back(13);
     	m[12].push_back(14);
-    	AutoRTFM::Abort();
+    	AutoRTFM::AbortTransaction();
 	});
 
     REQUIRE(
diff --git a/Engine/Source/Programs/AutoRTFMTests/Private/OpenAPI.cpp b/Engine/Source/Programs/AutoRTFMTests/Private/OpenAPI.cpp
new file mode 100644
--- /dev/null
+++ b/Engine/Source/Programs/AutoRTFMTests/Private/OpenAPI.cpp
@@ -0,0 +1,855 @@
+// Copyright Epic Games, Inc. All Rights Reserved.
+
+#pragma autortfm
+
+#include "Catch2Includes.h"
+#include <AutoRTFM/AutoRTFM.h>
+
+#include <map>
+#include <vector>
+
+TEST_CASE("OpenAPI.StartAbortAndStartAgain")
+{
+	int valueB = 0;
+	int valueC = 0;
+	AutoRTFM::Transact([&]()
+	{
+		// Recorded valueB as starting at 0.
+		valueB = 20;
+
+		AutoRTFM::Open([&]()
+		{
+			AutoRTFM::StartTransaction();
+			AutoRTFM::RecordOpenWrite(&valueB);
+			valueB=10;
+			AutoRTFM::AbortTransaction();
+
+			AutoRTFM::ClearTransactionStatus();
+
+			AutoRTFM::StartTransaction();
+			AutoRTFM::RecordOpenWrite(&valueC);
+			valueC = 30;
+			AutoRTFM::AbortTransaction();
+		});
+	});
+
+	// We rollback the transaction to the value we had when we first recorded the address
+	REQUIRE(valueB == 20);
+	REQUIRE(valueC == 0);
+}
+
+TEST_CASE("OpenAPI.CommitScopedFromOpen_Illegal", "[.]")
+{
+	AutoRTFM::ETransactionResult transactResult = AutoRTFM::Transact([&]()
+	{
+		AutoRTFM::Open([&]() {
+			AutoRTFM::CommitTransaction(); // illegal. Can't Commit from within a scoped transaction
+		});
+	});
+	REQUIRE(transactResult == AutoRTFM::ETransactionResult::AbortedByRequest);
+}
+
+TEST_CASE("OpenAPI.RecordDataClosed_Illegal", "[.]")
+{
+	int value = 0;
+	AutoRTFM::ETransactionResult transactResult = AutoRTFM::Transact([&]()
+	{
+		AutoRTFM::Close([&]() {
+			AutoRTFM::RecordOpenWrite(&value); // Illegal. Can't record writes explicitly while closed
+			value = 1;
+		});
+	});
+
+	REQUIRE(transactResult == AutoRTFM::ETransactionResult::Committed);
+	REQUIRE(value == 1);
+}
+
+TEST_CASE("OpenAPI.WriteDataInTheOpen")
+{
+	int value = 0;
+	auto transactResult = AutoRTFM::Transact([&]()
+	{
+		AutoRTFM::Open([&]() {
+			AutoRTFM::RecordOpenWrite(&value);
+			value = 1;
+		});
+	});
+
+	REQUIRE(transactResult == AutoRTFM::ETransactionResult::Committed);
+	REQUIRE(value == 1);
+}
+
+TEST_CASE("OpenAPI.AbortTransactionScopedFromOpen")
+{
+	auto transactResult = AutoRTFM::Transact([&]()
+	{
+		AutoRTFM::Open([&]() {
+			AutoRTFM::AbortTransaction();
+		});
+		FAIL("AutoRTFM::Open failed to throw after an abort");
+	});
+	REQUIRE(transactResult == AutoRTFM::ETransactionResult::AbortedByRequest);
+}
+
+TEST_CASE("OpenAPI.AbortTransactionScopedFromClosed")
+{
+	auto transactResult = AutoRTFM::Transact([&]()
+	{
+		AutoRTFM::Close([&]() {
+			AutoRTFM::AbortTransaction();
+		});
+		FAIL("AutoRTFM::Close should have no-op'ed because it's already closed from the Transact");
+	});
+	REQUIRE(transactResult == AutoRTFM::ETransactionResult::AbortedByRequest);
+}
+
+TEST_CASE("OpenAPI.AbortTransactionDoubleScopedFromOpen")
+{
+	unsigned long value = -42;
+	auto transactResult = AutoRTFM::Transact([&]()
+	{
+		value = 42;
+
+		auto transactResult2 = AutoRTFM::Transact([&]()
+		{
+			value = 42424242;
+			AutoRTFM::Open([&]() {
+				AutoRTFM::AbortTransaction();
+			});
+			FAIL("AutoRTFM::Open failed to throw after an abort");
+			value = 24242424;
+		});
+
+		if (transactResult2 != AutoRTFM::ETransactionResult::AbortedByRequest) FAIL("transactResult2 != AutoRTFM::ETransactionResult::AbortedByRequest");
+		if (value != 42) FAIL("value != 42");
+		value = 123123123;
+	});
+	REQUIRE(transactResult == AutoRTFM::ETransactionResult::Committed);
+	REQUIRE(value == 123123123);
+}
+
+TEST_CASE("OpenAPI.NestedClosedTransactions")
+{
+	int value = 0x12345678;
+	auto transactResult = AutoRTFM::Transact([&]()
+	{
+		// Read value
+		int x = value;
+		value = 0x11111111;
+
+		auto transactResult2 = AutoRTFM::Transact([&]()
+		{
+			if (value != 0x11111111) FAIL("value != 0x11111111");
+			// Read value
+			int y = value;
+			value = 0x22222222;
+
+			auto transactResult3 = AutoRTFM::Transact([&]()
+			{
+				if (value != 0x22222222) FAIL("value != 0x22222222");
+				// Read value
+				int z = value;
+				value = 0x33333333;
+
+				auto transactResult4 = AutoRTFM::Transact([&]()
+				{
+					if (value != 0x33333333) FAIL("value != 0x33333333");
+					// Read value
+					int q = value;
+					value = 0x44444444;
+					if (value != 0x44444444) FAIL("value != 0x44444444");
+					if (q != 0x33333333) FAIL("q != 0x33333333");
+				});
+				(void)transactResult4;
+
+				if (value != 0x44444444) FAIL("value != 0x44444444");
+				if (z != 0x22222222) FAIL("z != 0x22222222");
+			});
+			(void)transactResult3;
+
+			value = 0x55555555;
+			if (value != 0x55555555) FAIL("value != 0x55555555");
+			if (y != 0x11111111) FAIL("y != 0x11111111");
+		});
+		(void)transactResult2;
+		if (value != 0x55555555) FAIL("value != 0x55555555");
+
+		value = 0x66666666;
+		if (value != 0x66666666) FAIL("value != 0x66666666");
+		if (x != 0x12345678) FAIL("x != 0x12345678");
+	});
+
+	REQUIRE(transactResult == AutoRTFM::ETransactionResult::Committed);
+	REQUIRE(value == 0x66666666);
+}
+
+TEST_CASE("OpenAPI.OpenWithCopy")
+{
+	struct SomeData_t
+	{
+		int A;
+		float B;
+		char C;
+	};
+
+	SomeData_t SomeData1{ 1,2.0,'3' };
+
+	auto transactResult = AutoRTFM::Transact([&]()
+	{
+		SomeData_t SomeData2{ 9,8.0,'7' };
+		SomeData1.A = 11;
+		SomeData2.A = 29;
+
+		AutoRTFM::Open([=]() {
+			REQUIRE(SomeData1.A == 11);
+			REQUIRE(SomeData2.A == 29);
+		});
+	});
+
+	REQUIRE(transactResult == AutoRTFM::ETransactionResult::Committed);
+}
+
+#if defined(_BROKEN_ALLOC_FIXED_)
+
+TEST_CASE("OpenAPI.OpenCloseOpenClose")
+{
+	// START OPEN
+	REQUIRE(!AutoRTFM::IsTransactional());
+
+	int x = 42;
+	std::vector<int> v;
+	std::map<int, std::vector<int>> m;
+	v.push_back(100);
+	m[1].push_back(2);
+	m[1].push_back(3);
+	m[4].push_back(5);
+	m[6].push_back(7);
+	m[6].push_back(8);
+	m[6].push_back(9);
+
+	auto transactResult = AutoRTFM::Transact([&]()
+	{
+		// A - WE ARE CLOSED 
+		if (!AutoRTFM::IsClosed()) FAIL("A - NOT CLOSED AS EXPECTED!");
+
+		// -------------------------------------
+		AutoRTFM::Open([&]() {
+			// B - WE ARE OPEN 
+			REQUIRE(!AutoRTFM::IsClosed());
+
+			// -------------------------------------
+			AutoRTFM::Close([&]() {
+				// C - WE ARE CLOSED AGAIN
+				if (!AutoRTFM::IsClosed()) FAIL("C - NOT CLOSED AS EXPECTED!");
+
+				// -------------------------------------
+				AutoRTFM::Open([&]() {
+					// D - WE ARE OPEN AGAIN
+					REQUIRE(!AutoRTFM::IsClosed());
+
+					// An abort here will set state on the transactions, but will not LongJump
+					// AutoRTFM::Abort();
+				});
+
+				// -------------------------------------
+				// E - BACK TO CLOSED AFTER AN OPEN
+
+				x = 5;
+				for (size_t n = 10; n--;)
+					v.push_back(2 * n);
+				m.clear();
+				m[10].push_back(11);
+				m[12].push_back(13);
+				m[12].push_back(14);
+
+				// An abort here is closed and will LongJump past F AND G all the way to H
+				AutoRTFM::AbortTransaction();
+
+				// -------------------------------------
+				AutoRTFM::Open([&]() {
+					// F - WE ARE OPEN AGAIN //
+					REQUIRE(!AutoRTFM::IsClosed());
+				});
+
+				// -------------------------------------
+				// G - BACK TO CLOSED AGAIN
+				if (!AutoRTFM::IsClosed()) FAIL("NOT CLOSED!");
+
+			});
+			// -------------------------------------
+			// H - BACK TO OPEN 
+			REQUIRE(!AutoRTFM::IsClosed());
+
+		});
+
+		// -------------------------------------
+		// I - Finally closed again to finish out the transaction
+		if (!AutoRTFM::IsClosed()) FAIL("I - NOT CLOSED AS EXPECTED!");
+	});
+
+	REQUIRE(
+		AutoRTFM::ETransactionResult::AbortedByRequest ==
+		transactResult);
+	REQUIRE(x == 42);
+	REQUIRE(v.size() == 1);
+	REQUIRE(v[0] == 100);
+	REQUIRE(m.size() == 3);
+	REQUIRE(m[1].size() == 2);
+	REQUIRE(m[1][0] == 2);
+	REQUIRE(m[1][1] == 3);
+	REQUIRE(m[4].size() == 1);
+	REQUIRE(m[4][0] == 5);
+	REQUIRE(m[6].size() == 3);
+	REQUIRE(m[6][0] == 7);
+	REQUIRE(m[6][1] == 8);
+	REQUIRE(m[6][2] == 9);
+	REQUIRE(!AutoRTFM::IsTransactional());
+}
+
+#endif
+
+TEST_CASE("OpenAPI.Commit_TransactOpenCloseCommit")
+{
+	REQUIRE(!AutoRTFM::IsTransactional());
+
+	// We're open
+	REQUIRE(!AutoRTFM::IsClosed());
+
+	int value = 10;
+	value++;
+
+	// Close and start the top-level transaction
+	AutoRTFM::Transact([&]()
+	{
+		if (!AutoRTFM::IsClosed()) FAIL("Not Closed");
+
+		AutoRTFM::Open([&]() {
+			AutoRTFM::StartTransaction();
+
+			AutoRTFM::Close([&]() {
+				value = 42;
+			});
+
+			REQUIRE(value == 42); // RTFM writes through immediately, so we can see this value in the open
+			AutoRTFM::CommitTransaction();
+		});
+
+		if (value != 42) FAIL("Value != 42!");
+
+		value = 420;
+	});
+
+	REQUIRE(value == 420);
+
+	REQUIRE(!AutoRTFM::IsTransactional());
+}
+
+
+TEST_CASE("OpenAPI.Commit_TransactOpenCloseAbort")
+{
+	REQUIRE(!AutoRTFM::IsTransactional());
+
+	// We're open
+	REQUIRE(!AutoRTFM::IsClosed());
+
+	int value = 10;
+	value++;
+
+	// Close and start the top-level transaction
+	AutoRTFM::Transact([&]()
+	{
+		if (!AutoRTFM::IsClosed()) FAIL("Not Closed");
+
+		AutoRTFM::Open([&]()
+		{
+			double valueLocal = 1.0;
+			AutoRTFM::StartTransaction();
+
+			// Closing from the open doesn't work
+			AutoRTFM::Close([&]() {
+				value = 42;
+				valueLocal = 10.0;
+			});
+
+			AutoRTFM::AbortTransaction(); // undoes value = 42 in the open
+			REQUIRE(valueLocal == 1.0);
+		});
+
+		FAIL("Should not reach here!");
+	});
+
+	REQUIRE(value == 11);
+	REQUIRE(!AutoRTFM::IsTransactional());
+}
+
+TEST_CASE("OpenAPI.DoubleTransact")
+{
+	double value = 1.0;
+
+	AutoRTFM::Transact([&]()
+	{
+		AutoRTFM::Transact([&]()
+		{
+			value *= 2.5;
+			AutoRTFM::AbortTransaction();
+		});
+
+		value *= 10.0;
+	});
+
+	REQUIRE(value == 10.0);
+}
+
+TEST_CASE("OpenAPI.DoubleTransact2")
+{
+	double value = 1.0;
+
+	AutoRTFM::Transact([&]()
+	{
+		value = value + 2.0;
+		AutoRTFM::Transact([&]()
+		{
+			if (value == 3.0)
+			{
+				value *= 2.5;
+			}
+
+			if (value == 7.5)
+			AutoRTFM::AbortTransaction();
+		});
+
+		value *= 10.0;
+	});
+
+	REQUIRE(value == 30.0);
+}
+
+TEST_CASE("OpenAPI.DoubleTransact3")
+{
+	double result = 0.0;
+	AutoRTFM::Transact([&]()
+	{
+		double value = 1.0;
+		value = value + 2.0;
+		AutoRTFM::Transact([&]()
+		{
+			if (value == 3.0)
+			{
+				value *= 2.5;
+			}
+
+			if (value == 7.5)
+			AutoRTFM::AbortTransaction();
+		});
+
+		value *= 10.0;
+		result = value;
+	});
+
+	REQUIRE(result == 30.0);
+}
+
+TEST_CASE("OpenAPI.StackWriteCommitInTheOpen1")
+{
+	int value = 0;
+	AutoRTFM::Transact([&]()
+	{
+		AutoRTFM::Open([&]()
+		{
+			AutoRTFM::StartTransaction();
+			AutoRTFM::RecordOpenWrite(&value);
+			value = 10;
+			AutoRTFM::CommitTransaction();
+			REQUIRE(value == 10);
+		});
+	});
+}
+
+TEST_CASE("OpenAPI.StackWriteCommitInTheOpen2")
+{
+	AutoRTFM::Transact([&]()
+	{
+		int value = 0;
+		AutoRTFM::Open([&]()
+		{
+			AutoRTFM::StartTransaction();
+			AutoRTFM::Close([&]()
+			{
+				AutoRTFM::Open([&]() {
+					AutoRTFM::RecordOpenWrite(&value);
+					value = 10;
+				});
+			});
+
+			AutoRTFM::CommitTransaction();
+			REQUIRE(value == 10);
+		});
+	});
+}
+
+TEST_CASE("OpenAPI.StackWriteAbortInTheOpen1")
+{
+	int value = 0;
+	AutoRTFM::Transact([&]()
+	{
+		AutoRTFM::Open([&]()
+		{
+			AutoRTFM::StartTransaction();
+			AutoRTFM::RecordOpenWrite(&value);
+			value = 10;
+			AutoRTFM::AbortTransaction();
+			REQUIRE(value == 0);
+		});
+	});
+}
+
+#if defined(OPENAPI_ILLEGAL_TESTS)
+TEST_CASE("OpenAPI.StackWriteCommitInTheOpen3_Illegal")
+{
+	AutoRTFM::Transact([&]()
+	{
+		int value = 0;
+		AutoRTFM::Open([&]()
+		{
+			AutoRTFM::StartTransaction();
+			AutoRTFM::WriteMemory(&value, 10);
+			AutoRTFM::CommitTransaction();
+			REQUIRE(value == 10);
+		});
+	});
+}
+#endif
+
+int value1 = 0;
+
+TEST_CASE("OpenAPI.WriteMemory1")
+{
+	const int sourceValue = 10;
+	AutoRTFM::Transact([&]()
+	{
+		AutoRTFM::Open([&]()
+		{
+			AutoRTFM::StartTransaction();
+
+			AutoRTFM::WriteMemory(&value1, &sourceValue);
+			AutoRTFM::CommitTransaction();
+			REQUIRE(value1 == 10);
+		});
+	});
+}
+
+TEST_CASE("OpenAPI.StackWriteAbortInTheOpen2")
+{
+	int value = 0;
+	int bGotToA = false;
+	AutoRTFM::Transact([&]()
+	{
+		AutoRTFM::Open([&]()
+		{
+			AutoRTFM::StartTransaction();
+			AutoRTFM::WriteMemory(&value, 10); // Illegal to write to value because it's in the inner-most closed-nest
+
+			AutoRTFM::AbortTransaction();
+		});
+
+		// Never gets here
+		bGotToA = true;
+		REQUIRE(value == 0);
+	});
+
+	REQUIRE(bGotToA == false);
+	REQUIRE(value == 0);
+}
+
+TEST_CASE("OpenAPI.WriteTrivialStructure")
+{
+	struct SomeData
+	{
+		int A;
+		double B;
+		float C;
+		char D;
+		long E[5];
+	};
+
+	SomeData data = { 1,2.0, 3.0f, 'q', {123,234,345,456,567} };
+	SomeData data2 = { 9,8.0, 7.0f, '^', {999,888,777,666,555} };
+
+	AutoRTFM::Transact([&]()
+	{
+		AutoRTFM::Open([&]()
+		{
+			AutoRTFM::StartTransaction();
+
+			AutoRTFM::WriteMemory(&data, &data2);
+			REQUIRE(data.A == 9);
+			REQUIRE(data.B == 8.0);
+			REQUIRE(data.C == 7.0f);
+			REQUIRE(data.D == '^');
+			REQUIRE(data.E[0] == 999);
+			REQUIRE(data.E[1] == 888);
+			REQUIRE(data.E[2] == 777);
+			REQUIRE(data.E[3] == 666);
+			REQUIRE(data.E[4] == 555);
+			
+			AutoRTFM::AbortTransaction();
+			REQUIRE(data.A == 1);
+			REQUIRE(data.B == 2.0);
+			REQUIRE(data.C == 3.0f);
+			REQUIRE(data.D == 'q');
+			REQUIRE(data.E[0] == 123);
+			REQUIRE(data.E[1] == 234);
+			REQUIRE(data.E[2] == 345);
+			REQUIRE(data.E[3] == 456);
+			REQUIRE(data.E[4] == 567);
+		});
+	});
+}
+
+TEST_CASE("OpenAPI.WriteTrivialStructure2")
+{
+	struct SomeData
+	{
+		int A;
+		double B;
+		float C;
+		char D;
+		long E[5];
+	};
+
+	SomeData data = { 1,2.0, 3.0f, 'q', {123,234,345,456,567} };
+	SomeData data2 = { 9,8.0, 7.0f, '^', {999,888,777,666,555} };
+	SomeData data3 = { 19,28.0, 37.0f, '@', {4999,5888,6777,7666,8555} };
+
+	AutoRTFM::Transact([&]()
+	{
+		AutoRTFM::Open([&]()
+			{
+				AutoRTFM::StartTransaction();
+
+				AutoRTFM::WriteMemory(&data, &data2);
+				REQUIRE(data.A == 9);
+				REQUIRE(data.B == 8.0);
+				REQUIRE(data.C == 7.0f);
+				REQUIRE(data.D == '^');
+				REQUIRE(data.E[0] == 999);
+				REQUIRE(data.E[1] == 888);
+				REQUIRE(data.E[2] == 777);
+				REQUIRE(data.E[3] == 666);
+				REQUIRE(data.E[4] == 555);
+					
+				AutoRTFM::WriteMemory(&data, &data3);
+				REQUIRE(data.A == 19);
+				REQUIRE(data.B == 28.0);
+				REQUIRE(data.C == 37.0f);
+				REQUIRE(data.D == '@');
+				REQUIRE(data.E[0] == 4999);
+				REQUIRE(data.E[1] == 5888);
+				REQUIRE(data.E[2] == 6777);
+				REQUIRE(data.E[3] == 7666);
+				REQUIRE(data.E[4] == 8555);
+
+				AutoRTFM::AbortTransaction();
+				REQUIRE(data.A == 1);
+				REQUIRE(data.B == 2.0);
+				REQUIRE(data.C == 3.0f);
+				REQUIRE(data.D == 'q');
+				REQUIRE(data.E[0] == 123);
+				REQUIRE(data.E[1] == 234);
+				REQUIRE(data.E[2] == 345);
+				REQUIRE(data.E[3] == 456);
+				REQUIRE(data.E[4] == 567);
+			});
+	});
+}
+
+TEST_CASE("OpenAPI.Footgun1")
+{
+	int valueA = 0;
+	int valueB = 0;
+
+	AutoRTFM::Transact([&]()
+	{
+		// Does nothing - already closed
+		AutoRTFM::Close([&]()
+		{
+			// Recorded valueB as starting at 0.
+			valueB = 123;
+
+			AutoRTFM::Open([&]()
+			{
+				// Unrecorded assignments in the open
+				valueA = 10;
+				valueB = 10;
+			});
+				
+			// valueA is now recorded as starting at 10
+			valueA = 20;
+			AutoRTFM::AbortTransaction();
+		});
+	});
+
+	// We rollback the transaction to the value we had when we first recorded the address
+	REQUIRE(valueA == 10);
+	REQUIRE(valueB == 0);
+}
+
+TEST_CASE("OpenAPI.Footgun2")
+{
+	int valueB = 0;
+	int valueC = 0;
+	AutoRTFM::Transact([&]()
+	{
+		// Does nothing - already closed
+		AutoRTFM::Close([&]()
+		{
+			// Recorded valueB as starting at 0.
+			valueB = 20;
+
+			AutoRTFM::Open([&]()
+			{
+				// Unrecorded assignments in the open
+				valueB = 10;
+				valueC = 10;
+				AutoRTFM::RecordOpenWrite(&valueC);
+				// valueC was recorded in the open after the change - too late
+			});
+
+			// valueA is now recorded as starting at 10
+			valueC = 40;
+			AutoRTFM::AbortTransaction();
+		});
+	});
+
+	// We rollback the transaction to the value we had when we first recorded the address
+	REQUIRE(valueB == 0);
+	REQUIRE(valueC == 10);
+}
+
+#if 0
+TEST_CASE("OpenAPI.StartCloseOpenCommit")
+{
+	REQUIRE(!AutoRTFM::IsTransactional());
+
+	int value = 10;
+	value++;
+	(void)value;
+
+	AutoRTFM::StartTransaction();
+
+	// Can't close outside of a transaction
+	AutoRTFM::Close([&]() {
+		value = 420;
+	});
+
+	// assignment within the close should be visible to us
+	REQUIRE(value == 420);
+
+	// Setting a value in the open requires us to register the memory address with the transaction
+	value = 42;
+	AutoRTFM::RecordOpenWrite(&value);
+
+	AutoRTFM::CommitTransaction();
+
+	// Finally, 42 is committed to value
+	REQUIRE(value == 42);
+
+	REQUIRE(!AutoRTFM::IsTransactional());
+}
+#endif
+
+TEST_CASE("OpenAPI.TransOpenStartCloseAbortAbort")
+{
+	bool bGetsToA = false;
+	bool bGetsToB = false;
+	bool bGetsToC = false;
+	bool bGetsToD = false;
+
+	REQUIRE(!AutoRTFM::IsTransactional());
+
+	AutoRTFM::Transact([&]() 
+	{
+		AutoRTFM::Open([&]() {
+			AutoRTFM::StartTransaction();
+
+			int value = 10;
+			value++;
+
+			value = 42;
+			AutoRTFM::EContextStatus CloseStatus = AutoRTFM::Close([&]()
+			{
+				value = 420;
+				AutoRTFM::AbortTransaction();
+				bGetsToA = true;
+			});
+
+			REQUIRE(CloseStatus == AutoRTFM::EContextStatus::AbortedByRequest);
+
+			REQUIRE(AutoRTFM::IsAborting());
+			AutoRTFM::ClearTransactionStatus();
+			REQUIRE(!AutoRTFM::IsAborting());
+
+			REQUIRE(bGetsToA == false);
+
+			bGetsToB = true;
+			REQUIRE(value == 42);
+			AutoRTFM::AbortTransaction();
+			bGetsToC = true;
+			REQUIRE(value == 42);
+		});
+
+		bGetsToD = true;
+	});
+
+	REQUIRE(bGetsToA == false);
+	REQUIRE(bGetsToB == true);
+	REQUIRE(bGetsToC == true);
+	REQUIRE(bGetsToD == false);
+	REQUIRE(!AutoRTFM::IsTransactional());
+}
+
+TEST_CASE("OpenAPI.TransOpenTransCloseAbortAbort")
+{
+	bool bGetsToA = false;
+	bool bGetsToB = false;
+	bool bGetsToC = false;
+	bool bGetsToD = false;
+
+	REQUIRE(!AutoRTFM::IsTransactional());
+
+	AutoRTFM::Transact([&]() 
+	{
+		AutoRTFM::Open([&]() {
+			AutoRTFM::Transact([&](){
+				int value = 10;
+				value++;
+
+				value = 42;
+				// Can't close outside of a Transact
+				AutoRTFM::Close([&]()
+				{
+					value = 420;
+					AutoRTFM::AbortTransaction();
+					bGetsToA = true;
+				});
+
+				REQUIRE(bGetsToA == false);
+
+				bGetsToB = true;
+				REQUIRE(value == 42);
+				AutoRTFM::AbortTransaction();
+				bGetsToC = true;
+				REQUIRE(value == 42);
+			});
+		});
+
+		bGetsToD = true;
+	});
+
+	REQUIRE(bGetsToA == false);
+	REQUIRE(bGetsToB == false);
+	REQUIRE(bGetsToC == false);
+	REQUIRE(bGetsToD == true);
+	REQUIRE(!AutoRTFM::IsTransactional());
+}
+
diff --git a/Engine/Source/Runtime/Core/Private/AutoRTFM/API.cpp b/Engine/Source/Runtime/Core/Private/AutoRTFM/API.cpp
--- a/Engine/Source/Runtime/Core/Private/AutoRTFM/API.cpp
+++ b/Engine/Source/Runtime/Core/Private/AutoRTFM/API.cpp
@@ -4,8 +4,10 @@
 
 #include "AutoRTFM/AutoRTFM.h"
 #include "AutoRTFM/AutoRTFMConstants.h"
+#include "CallNest.h"
 #include "Context.h"
 #include "ContextInlines.h"
+#include "ContextStatus.h"
 #include "Debug.h"
 #include "FunctionMapInlines.h"
 #include "TransactionInlines.h"
@@ -56,6 +58,61 @@ extern "C" void autortfm_abort()
     }
 }
 
+extern "C" bool autortfm_start_transaction()
+{
+	if (FContext::IsTransactional())
+	{
+		return FContext::Get()->StartTransaction();
+	}
+	else
+	{
+		fprintf(stderr, "autortfm_start_transaction called from outside a transact.\n");
+		abort();
+	}
+}
+
+extern "C" autortfm_result autortfm_commit_transaction()
+{
+	if (FContext::IsTransactional())
+	{
+		return static_cast<autortfm_result>(FContext::Get()->CommitTransaction());
+	}
+	else
+	{
+		fprintf(stderr, "autortfm_commit_transaction called from outside a transaction.\n");
+		abort();
+	}
+}
+
+extern "C" autortfm_result autortfm_abort_transaction()
+{
+	if (FContext::IsTransactional())
+	{
+		return static_cast<autortfm_result>(FContext::Get()->AbortTransaction(false));
+	}
+	else
+	{
+		fprintf(stderr, "autortfm_abort_transaction called from outside a transaction.\n");
+		abort();
+	}
+}
+
+extern "C" void autortfm_clear_transaction_status()
+{
+	ASSERT(FContext::Get()->IsAborting());
+	FContext::Get()->ClearTransactionStatus();
+}
+
+extern "C" bool autortfm_is_aborting()
+{
+	return FContext::Get()->IsAborting();
+}
+
+extern "C" bool autortfm_current_nest_throw()
+{
+	FContext::Get()->Throw();
+}
+
 extern "C" void autortfm_abort_if_transactional()
 {
     if (FContext::IsTransactional())
@@ -71,18 +128,20 @@ extern "C" void autortfm_abort_if_closed()
 
 extern "C" void autortfm_open(void (*Work)(void* Arg), void* Arg)
 {
-    Work(Arg);
+	Work(Arg);
 }
 
-extern "C" void autortfm_close(void (*Work)(void* Arg), void* Arg)
+extern "C" autortfm_status autortfm_close(void (*Work)(void* Arg), void* Arg)
 {
+	autortfm_status Result = autortfm_status_ontrack;
+
     if (FContext::IsTransactional())
     {
         FContext* Context = FContext::Get();
         void (*WorkClone)(void* Arg, FContext* Context) = FunctionMapLookup(Work, Context, "autortfm_close");
         if (WorkClone)
         {
-            WorkClone(Arg, Context);
+			Result = static_cast<autortfm_status>(Context->CallClosedNest(WorkClone, Arg));
         }
     }
     else
@@ -90,6 +149,13 @@ extern "C" void autortfm_close(void (*Work)(void* Arg), void* Arg)
         fprintf(stderr, "autortfm_close called from outside a transaction.\n");
         abort();
     }
+
+	return Result;
+}
+
+extern "C" void autortfm_record_open_write(void* Ptr, size_t Size)
+{
+	FContext::Get()->RecordWrite(Ptr, Size, false);
 }
 
 extern "C" void autortfm_register_open_function(void* OriginalFunction, void* NewFunction)
@@ -259,6 +325,52 @@ void STM_autortfm_abort(FContext* Context)
 }
 UE_AUTORTFM_REGISTER_OPEN_FUNCTION(autortfm_abort);
 
+void STM_autortfm_start_transaction(FContext* Context)
+{
+	fprintf(stderr, "autortfm_start_transaction called from closed code.\n");
+	abort();
+}
+UE_AUTORTFM_REGISTER_OPEN_FUNCTION(autortfm_start_transaction);
+
+void STM_autortfm_commit_transaction(FContext* Context)
+{
+	fprintf(stderr, "autostm_committransaction called from closed code.\n");
+	abort();
+}
+UE_AUTORTFM_REGISTER_OPEN_FUNCTION(autortfm_commit_transaction);
+
+autortfm_result STM_autortfm_abort_transaction(FContext* Context)
+{
+	if (FContext::IsTransactional())
+	{
+		return static_cast<autortfm_result>(Context->AbortTransaction(true));
+	}
+	else
+	{
+		fprintf(stderr, "autortfm_abort_transaction called from outside a transaction.\n");
+		abort();
+	}
+}
+UE_AUTORTFM_REGISTER_OPEN_FUNCTION(autortfm_abort_transaction);
+
+void STM_autortfm_clear_transaction_status(FContext* Context)
+{
+	fprintf(stderr, "autortfm_clear_transaction_status called from closed code.\n");
+	abort();
+}
+
+bool STM_autortfm_is_aborting(FContext* Context)
+{
+	return Context->IsAborting();
+}
+UE_AUTORTFM_REGISTER_OPEN_FUNCTION(autortfm_is_aborting);
+
+bool STM_autortfm_current_nest_throw(FContext* Context)
+{
+	Context->Throw();
+}
+UE_AUTORTFM_REGISTER_OPEN_FUNCTION(autortfm_current_nest_throw);
+
 void STM_autortfm_abort_if_transactional(FContext* Context)
 {
     if (FDebug::bVerbose)
@@ -282,20 +394,30 @@ UE_AUTORTFM_REGISTER_OPEN_FUNCTION(autortfm_abort_if_closed);
 
 void STM_autortfm_open(void (*Work)(void* Arg), void* Arg, FContext* Context)
 {
-    Work(Arg);
+	// WARNING! DO NOT EDIT! Changes to this function will be elided due to special compiler optimizations
+	Work(Arg);
 }
 UE_AUTORTFM_REGISTER_OPEN_FUNCTION(autortfm_open);
 
-void STM_autortfm_close(void (*Work)(void* Arg), void* Arg, FContext* Context)
+autortfm_status STM_autortfm_close(void (*Work)(void* Arg), void* Arg, FContext* Context)
 {
     void (*WorkClone)(void* Arg, FContext* Context) = FunctionMapLookup(Work, Context, "STM_autortfm_close");
     if (WorkClone)
     {
         WorkClone(Arg, Context);
     }
+
+	return static_cast<autortfm_status>(FContext::Get()->GetStatus());
 }
 UE_AUTORTFM_REGISTER_OPEN_FUNCTION(autortfm_close);
 
+extern "C" void STM_autortfm_record_open_write(void*, size_t, FContext*)
+{
+	fprintf(stderr, "FATAL: autortfm_record_open_write called from closed code.\n");
+	abort();
+}
+UE_AUTORTFM_REGISTER_OPEN_FUNCTION(autortfm_record_open_write);
+
 void STM_DeferUntilCommit(std::function<void()>&& Work, FContext* Context)
 {
     ASSERT(Context->GetStatus() == EContextStatus::OnTrack);
diff --git a/Engine/Source/Runtime/Core/Private/AutoRTFM/CallNest.h b/Engine/Source/Runtime/Core/Private/AutoRTFM/CallNest.h
new file mode 100644
--- /dev/null
+++ b/Engine/Source/Runtime/Core/Private/AutoRTFM/CallNest.h
@@ -0,0 +1,29 @@
+// Copyright Epic Games, Inc. All Rights Reserved.
+
+#pragma once
+
+#include "LongJump.h"
+
+namespace AutoRTFM
+{
+
+class FContext;
+
+class FCallNest final
+{
+public:
+	FCallNest(FContext* Context)
+		: Context(Context)
+	{}
+
+	FCallNest* Parent{ nullptr };
+	FContext* Context;
+	FLongJump AbortJump;
+
+	// Whether this succeeded or not is reflected in Context::GetStatus().
+	template<typename TTryFunctor>
+	void Try(const TTryFunctor& TryFunctor);
+};
+
+
+} // namespace AutoRTFM
diff --git a/Engine/Source/Runtime/Core/Private/AutoRTFM/CallNestInlines.h b/Engine/Source/Runtime/Core/Private/AutoRTFM/CallNestInlines.h
new file mode 100644
--- /dev/null
+++ b/Engine/Source/Runtime/Core/Private/AutoRTFM/CallNestInlines.h
@@ -0,0 +1,27 @@
+// Copyright Epic Games, Inc. All Rights Reserved.
+
+#pragma once
+
+#include "Context.h"
+#include "CallNest.h"
+
+namespace AutoRTFM
+{
+
+template<typename TTryFunctor>
+void FCallNest::Try(const TTryFunctor& TryFunctor)
+{
+	AbortJump.TryCatch(
+		[&]()
+		{
+			TryFunctor();
+			ASSERT(Context->GetStatus() == EContextStatus::OnTrack);
+		},
+		[&]()
+		{
+			ASSERT(Context->GetStatus() != EContextStatus::Idle);
+			ASSERT(Context->GetStatus() != EContextStatus::OnTrack);
+		});
+}
+
+} // namespace AutoRTFM
diff --git a/Engine/Source/Runtime/Core/Private/AutoRTFM/CompilerInterface.cpp b/Engine/Source/Runtime/Core/Private/AutoRTFM/CompilerInterface.cpp
--- a/Engine/Source/Runtime/Core/Private/AutoRTFM/CompilerInterface.cpp
+++ b/Engine/Source/Runtime/Core/Private/AutoRTFM/CompilerInterface.cpp
@@ -31,7 +31,7 @@ void CheckAlignment(FContext* Context, void* Ptr, size_t AlignmentMask)
 
 extern "C" void autortfm_record_write(FContext* Context, void* Ptr, size_t Size)
 {
-    Context->RecordWrite(Ptr, Size);
+    Context->RecordWrite(Ptr, Size, true);
 }
 
 extern "C" void* autortfm_lookup_function(FContext* Context, void* OriginalFunction, const char* Where)
diff --git a/Engine/Source/Runtime/Core/Private/AutoRTFM/Context.cpp b/Engine/Source/Runtime/Core/Private/AutoRTFM/Context.cpp
--- a/Engine/Source/Runtime/Core/Private/AutoRTFM/Context.cpp
+++ b/Engine/Source/Runtime/Core/Private/AutoRTFM/Context.cpp
@@ -2,6 +2,7 @@
 
 #if (defined(__AUTORTFM) && __AUTORTFM)
 #include "Context.h"
+#include "CallNestInlines.h"
 #include "ContextInlines.h"
 #include "FunctionMap.h"
 #include "GlobalData.h"
@@ -61,6 +62,194 @@ bool FContext::IsTransactional()
     }
 }
 
+bool FContext::StartTransaction()
+{
+	FTransaction* NewTransaction = new FTransaction(this);
+
+	void* TransactStackAddress = &NewTransaction; // get the stack-local position pointer.
+	ASSERT(TransactStackAddress > StackBegin);
+	ASSERT(TransactStackAddress < StackEnd);
+	ASSERT(TransactStackAddress < CurrentTransactStackAddress);
+
+	// This form of transaction is always ultimately within a scoped Transact 
+	ASSERT(Status == EContextStatus::OnTrack);
+	PushTransaction(NewTransaction);
+
+	return true;
+}
+
+ETransactionResult FContext::CommitTransaction()
+{
+	constexpr bool bVerbose = false;
+
+	ASSERT(Status == EContextStatus::OnTrack);
+
+	// Scoped transactions commit on return, so committing explicitly isn't allowed
+	ASSERT(CurrentTransaction->IsScopedTransaction() == false);
+
+	ETransactionResult Result = ETransactionResult::Committed;
+
+	if (CurrentTransaction->IsNested())
+	{
+		Result = ResolveNestedTransaction(CurrentTransaction);
+	}
+	else
+	{
+		if (bVerbose)
+		{
+			fprintf(GetLogFile(), "About to commit; my state is:\n");
+			DumpState();
+			fprintf(GetLogFile(), "Committing...\n");
+		}
+
+		if (CurrentTransaction->AttemptToCommit())
+		{
+			Result = ETransactionResult::Committed;
+		}
+		else
+		{
+			if (bVerbose)
+			{
+				fprintf(GetLogFile(), "Commit failed!\n");
+			}
+			ASSERT(Status != EContextStatus::OnTrack);
+			ASSERT(Status != EContextStatus::Idle);
+		}
+	}
+
+	// Parent transaction is now the current transaction
+	PopTransaction();
+
+	return Result;
+}
+
+ETransactionResult FContext::AbortTransaction(bool bIsClosed)
+{
+	ETransactionResult Result = ETransactionResult::AbortedByRequest;
+	ASSERT(Status == EContextStatus::OnTrack);
+	Status = EContextStatus::AbortedByRequest;
+
+	ASSERT(CurrentTransaction != nullptr);
+
+	// Sort out how aborts work
+	CurrentTransaction->AbortWithoutThrowing();
+
+	// Non-scoped transactions are ended immediately, but scoped need to get to the end scope before being popped
+	if (!CurrentTransaction->IsScopedTransaction())
+	{
+		Result = ResolveNestedTransaction(CurrentTransaction);
+		PopTransaction();
+	}
+	
+	if (bIsClosed)
+	{
+		Throw();
+	}
+
+	return Result;
+}
+
+bool FContext::IsAborting() const
+{
+	return Status != EContextStatus::OnTrack && Status != EContextStatus::Idle;
+}
+
+EContextStatus FContext::CallClosedNest(void (*ClosedFunction)(void* Arg, FContext* Context), void* Arg)
+{
+	TScopedGuard<void*> CurrentNestStackAddressGuard(CurrentTransactStackAddress, &CurrentNestStackAddressGuard);
+
+	PushCallNest(new FCallNest(this));
+
+	CurrentNest->Try([&]() { ClosedFunction(Arg, this); });
+
+	PopCallNest();
+
+	return GetStatus();
+}
+
+void FContext::PushCallNest(FCallNest* NewCallNest)
+{
+	ASSERT(NewCallNest != nullptr);
+	ASSERT(NewCallNest->Parent == nullptr);
+
+	NewCallNest->Parent = CurrentNest;
+	CurrentNest = NewCallNest;
+}
+
+void FContext::PopCallNest()
+{
+	ASSERT(CurrentNest != nullptr);
+	FCallNest* OldCallNest = CurrentNest;
+	CurrentNest = CurrentNest->Parent;
+
+	delete OldCallNest;
+}
+
+void FContext::PushTransaction(FTransaction* NewTransaction)
+{
+	ASSERT(NewTransaction != nullptr);
+	ASSERT(!NewTransaction->IsDone());
+	ASSERT(NewTransaction->GetParent() == nullptr);
+
+	ASSERT(CurrentTransaction == nullptr || !CurrentTransaction->IsDone());
+	
+	NewTransaction->SetParent(CurrentTransaction);
+	CurrentTransaction = NewTransaction;
+}
+
+void FContext::PopTransaction()
+{
+	ASSERT(CurrentTransaction != nullptr);
+	ASSERT(CurrentTransaction->IsDone());
+	FTransaction* OldTransaction = CurrentTransaction;
+	CurrentTransaction = CurrentTransaction->GetParent();
+	delete OldTransaction;
+}
+
+void FContext::ClearTransactionStatus()
+{
+	switch (Status)
+	{
+	case EContextStatus::OnTrack:
+		break;
+	case EContextStatus::AbortedByLanguage:
+		Status = EContextStatus::OnTrack;
+		break;
+	case EContextStatus::AbortedByRequest:
+		Status = EContextStatus::OnTrack;
+		break;
+	default:
+		ASSERT(!"Should not be reached");
+	}
+}
+
+ETransactionResult FContext::ResolveNestedTransaction(FTransaction* NewTransaction)
+{
+	// We just use this bit to help assertions for now (though we could use it more strongly). Because of how we use this right now,
+	// it's OK that it's set before we commit but after we abort.
+	ASSERT(!NewTransaction->IsDone());
+	NewTransaction->SetIsDone();
+
+	if (Status == EContextStatus::OnTrack)
+	{
+		bool bCommitResult = NewTransaction->AttemptToCommit();
+		ASSERT(bCommitResult);
+		ASSERT(Status == EContextStatus::OnTrack);
+		return ETransactionResult::Committed;
+	}
+
+	switch (Status)
+	{
+	case EContextStatus::AbortedByRequest:
+		return ETransactionResult::AbortedByRequest;
+	case EContextStatus::AbortedByLanguage:
+		return ETransactionResult::AbortedByLanguage;
+	default:
+		ASSERT(!"Should not be reached");
+		return ETransactionResult::AbortedByLanguage;
+	}
+}
+
 ETransactionResult FContext::Transact(void (*Function)(void* Arg), void* Arg)
 {
     constexpr bool bVerbose = false;
@@ -74,26 +263,36 @@ ETransactionResult FContext::Transact(void (*Function)(void* Arg), void* Arg)
         return ETransactionResult::AbortedByLanguage;
     }
     
-    std::unique_ptr<FTransaction> NewTransactionUniquePtr(new FTransaction(this));
-    FTransaction* NewTransaction = NewTransactionUniquePtr.get();
+	//std::unique_ptr<FTransaction> NewTransactionUniquePtr(new FTransaction(this));
+	FTransaction* NewTransaction = new FTransaction(this);
+	FCallNest* NewNest = new FCallNest(this);
 
-    void* TransactStackAddress = &NewTransactionUniquePtr;
-    ASSERT(TransactStackAddress > StackBegin);
-    ASSERT(TransactStackAddress < StackEnd);
-    TScopedGuard<void*> CurrentTransactStackAddressGuard(CurrentTransactStackAddress, TransactStackAddress);
+	// Transact requires a return from the lambda to commit the results
+	NewTransaction->SetIsScopedTransaction();
+
+	void* TransactStackAddress = &NewTransaction;
+	ASSERT(TransactStackAddress > StackBegin);
+	ASSERT(TransactStackAddress < StackEnd);
+	TScopedGuard<void*> CurrentNestStackAddressGuard(CurrentTransactStackAddress, TransactStackAddress);
+
+	ETransactionResult Result = ETransactionResult::Committed; // Initialize to something to make the compiler happy.
 
     if (!CurrentTransaction)
     {
         ASSERT(Status == EContextStatus::Idle);
-        CurrentTransaction = NewTransaction;
+
+		PushTransaction(NewTransaction);
+		PushCallNest(NewNest);
         OuterTransactStackAddress = TransactStackAddress;
-        ETransactionResult Result = ETransactionResult::Committed; // Initialize to something to make the compiler happy.
+
         for (;;)
         {
             Status = EContextStatus::OnTrack;
             ASSERT(CurrentTransaction->IsFresh());
-            CurrentTransaction->Try([&] () { ClonedFunction(Arg, this); });
+			CurrentNest->Try([&] () { ClonedFunction(Arg, this); });
+			ASSERT(CurrentTransaction == NewTransaction); // The transaction lambda should have unwound any nested transactions.
             ASSERT(Status != EContextStatus::Idle);
+
             if (Status == EContextStatus::OnTrack)
             {
                 if (bVerbose)
@@ -127,9 +326,7 @@ ETransactionResult FContext::Transact(void (*Function)(void* Arg), void* Arg)
                 Result = ETransactionResult::AbortedByLanguage;
                 break;
             }
-            ASSERT(Status == EContextStatus::AbortedByLineInitialization
-                   || Status == EContextStatus::AbortedByCommitTimeLineCheck
-                   || Status == EContextStatus::AbortedByFailedLockAcquisition);
+            ASSERT(Status == EContextStatus::AbortedByFailedLockAcquisition);
 
             if (bVerbose)
             {
@@ -140,47 +337,38 @@ ETransactionResult FContext::Transact(void (*Function)(void* Arg), void* Arg)
             bLocksToHoldAreHeld = true;
         }
 
+		NewTransaction->SetIsDone();
+
+		PopCallNest();
+		PopTransaction();
+		ClearTransactionStatus();
+
+		ASSERT(CurrentNest == nullptr);
+		ASSERT(CurrentTransaction == nullptr);
+
         Reset();
-        // No need to SetIsDone or whatever since all of the transactions have now been blown away.
-        return Result;
-    }
+	}
     else
     {
-        ASSERT(Status == EContextStatus::OnTrack);
-        FTransaction* PreviousTransaction = CurrentTransaction;
-        CurrentTransaction = NewTransaction;
-        CurrentTransaction->Try([&] () { ClonedFunction(Arg, this); });
-
-        // We just use this bit to help assertions for now (though we could use it more strongly). Because of how we use this right now,
-        // it's OK that it's set before we commit but after we abort.
-        NewTransaction->SetIsDone();
-        
-        if (Status == EContextStatus::OnTrack)
-        {
-            bool bCommitResult = NewTransaction->AttemptToCommit();
-            CurrentTransaction = PreviousTransaction;
-            ASSERT(bCommitResult);
-            ASSERT(Status == EContextStatus::OnTrack);
-            return ETransactionResult::Committed;
-        }
-        
-        CurrentTransaction = PreviousTransaction;
+		// This transaction is within another transaction
+		ASSERT(Status == EContextStatus::OnTrack);
+		PushTransaction(NewTransaction);
+		PushCallNest(NewNest);
 
-        switch (Status)
-        {
-        case EContextStatus::AbortedByLineInitialization:
-            CurrentTransaction->AbortAndThrow();
-            AutoRTFM::Unreachable();
-        case EContextStatus::AbortedByRequest:
-            Status = EContextStatus::OnTrack;
-            return ETransactionResult::AbortedByRequest;
-        case EContextStatus::AbortedByLanguage:
-            Status = EContextStatus::OnTrack;
-            return ETransactionResult::AbortedByLanguage;
-        default:
-            AutoRTFM::Unreachable();
-        }
-    }
+		CurrentNest->Try([&]() { ClonedFunction(Arg, this); });
+		ASSERT(CurrentTransaction == NewTransaction);
+
+		Result = ResolveNestedTransaction(NewTransaction);
+		
+		PopCallNest();
+		PopTransaction();
+		ClearTransactionStatus();
+
+		ASSERT(CurrentNest != nullptr);
+		ASSERT(CurrentTransaction != nullptr);
+	}
+
+	return Result;
 }
 
 void FContext::AbortByRequestAndThrow()
@@ -190,6 +378,13 @@ void FContext::AbortByRequestAndThrow()
     CurrentTransaction->AbortAndThrow();
 }
 
+void FContext::AbortByRequestWithoutThrowing()
+{
+	ASSERT(Status == EContextStatus::OnTrack);
+	Status = EContextStatus::AbortedByRequest;
+	CurrentTransaction->AbortWithoutThrowing();
+}
+
 void FContext::AbortByLanguageAndThrow()
 {
     constexpr bool bAbortProgram = false;
@@ -228,11 +423,17 @@ FContext::FContext()
 void FContext::Reset()
 {
     OuterTransactStackAddress = nullptr;
-    CurrentTransactStackAddress = nullptr;
+	CurrentTransactStackAddress = nullptr;
     CurrentTransaction = nullptr;
+	CurrentNest = nullptr;
     Status = EContextStatus::Idle;
 }
 
+void FContext::Throw()
+{
+	GetCurrentNest()->AbortJump.Throw();
+}
+
 void FContext::DumpState() const
 {
     fprintf(GetLogFile(), "Context at %p.\n", this);
diff --git a/Engine/Source/Runtime/Core/Private/AutoRTFM/Context.h b/Engine/Source/Runtime/Core/Private/AutoRTFM/Context.h
--- a/Engine/Source/Runtime/Core/Private/AutoRTFM/Context.h
+++ b/Engine/Source/Runtime/Core/Private/AutoRTFM/Context.h
@@ -4,13 +4,13 @@
 
 #include "AutoRTFM/AutoRTFM.h"
 #include "ContextStatus.h"
-#include <memory>
 
 namespace AutoRTFM
 {
 
 class FLineLock;
 class FTransaction;
+class FCallNest;
 
 class FContext
 {
@@ -21,10 +21,22 @@ public:
     
     // This is public API
     ETransactionResult Transact(void (*Function)(void* Arg), void* Arg);
-    void AbortByRequestAndThrow();
+    
+	EContextStatus CallClosedNest(void (*ClosedFunction)(void* Arg, FContext* Context), void* Arg);
+
+	void AbortByRequestAndThrow();
+	void AbortByRequestWithoutThrowing();
+
+	// Open API - no throw
+	bool StartTransaction();
+
+	ETransactionResult CommitTransaction();
+	ETransactionResult AbortTransaction(bool bIsClosed);
+	void ClearTransactionStatus();
+	bool IsAborting() const;
 
     // Record that a write is about to occur at the given LogicalAddress of Size bytes.
-    void RecordWrite(void* LogicalAddress, size_t Size);
+    void RecordWrite(void* LogicalAddress, size_t Size, bool bIsClosed);
 
     void DidAllocate(void* LogicalAddress, size_t Size);
     void WillDeallocate(void* LogicalAddress, size_t Size);
@@ -32,11 +44,13 @@ public:
     // The rest of this is internalish.
     void AbortByLanguageAndThrow();
 
-    FTransaction* GetCurrentTransaction() const { return CurrentTransaction; }
-    bool IsTransactionStack(void* LogicalAddress) const { return LogicalAddress >= StackBegin && LogicalAddress < OuterTransactStackAddress; }
-    bool IsInnerTransactionStack(void* LogicalAddress) const { return LogicalAddress >= StackBegin && LogicalAddress < CurrentTransactStackAddress; }
-    EContextStatus GetStatus() const { return Status; }
-
+	inline FTransaction* GetCurrentTransaction() const { return CurrentTransaction; }
+	inline FCallNest* GetCurrentNest() const { return CurrentNest; }
+    inline bool IsTransactionStack(void* LogicalAddress) const { return LogicalAddress >= StackBegin && LogicalAddress < OuterTransactStackAddress; }
+	inline bool IsInnerTransactionStack(void* LogicalAddress) const { return LogicalAddress >= StackBegin && LogicalAddress < CurrentTransactStackAddress; }
+	inline EContextStatus GetStatus() const { return Status; }
+	void Throw();
+	
     void DumpState() const;
 
     static void InitializeGlobalData();
@@ -45,12 +59,22 @@ private:
     FContext();
     FContext(const FContext&) = delete;
 
+	void PushCallNest(FCallNest* NewCallNest);
+	void PopCallNest();
+
+	void PushTransaction(FTransaction* NewTransaction);
+	void PopTransaction();
+
+	ETransactionResult ResolveNestedTransaction(FTransaction* NewTransaction);
+
     void Set();
     
     // All of this other stuff ought to be private?
     void Reset();
     
     FTransaction* CurrentTransaction{nullptr};
+	FCallNest* CurrentNest{nullptr};
+
     void* StackBegin{nullptr}; // begin as in the smaller of the two
     void* StackEnd{nullptr};
     void* OuterTransactStackAddress{nullptr};
diff --git a/Engine/Source/Runtime/Core/Private/AutoRTFM/ContextInlines.h b/Engine/Source/Runtime/Core/Private/AutoRTFM/ContextInlines.h
--- a/Engine/Source/Runtime/Core/Private/AutoRTFM/ContextInlines.h
+++ b/Engine/Source/Runtime/Core/Private/AutoRTFM/ContextInlines.h
@@ -8,9 +8,9 @@
 namespace AutoRTFM
 {
 
-inline void FContext::RecordWrite(void* LogicalAddress, size_t Size)
+inline void FContext::RecordWrite(void* LogicalAddress, size_t Size, bool bIsClosed)
 {
-    CurrentTransaction->RecordWrite(LogicalAddress, Size);
+    CurrentTransaction->RecordWrite(LogicalAddress, Size, bIsClosed);
 }
 
 inline void FContext::DidAllocate(void* LogicalAddress, size_t Size)
diff --git a/Engine/Source/Runtime/Core/Private/AutoRTFM/ContextStatus.h b/Engine/Source/Runtime/Core/Private/AutoRTFM/ContextStatus.h
--- a/Engine/Source/Runtime/Core/Private/AutoRTFM/ContextStatus.h
+++ b/Engine/Source/Runtime/Core/Private/AutoRTFM/ContextStatus.h
@@ -7,18 +7,6 @@
 
 namespace AutoRTFM
 {
-
-enum class EContextStatus
-{
-    Idle,
-    OnTrack,
-    AbortedByLineInitialization,
-    AbortedByCommitTimeLineCheck,
-    AbortedByFailedLockAcquisition,
-    AbortedByLanguage,
-    AbortedByRequest
-};
-
 static_assert(static_cast<unsigned>(EContextStatus::OnTrack) == Constants::Context_Status_OnTrack, "Not equal");
 
 inline const char* GetContextStatusName(EContextStatus Status)
@@ -29,10 +17,6 @@ inline const char* GetContextStatusName(EContextStatus Status)
         return "Idle";
     case EContextStatus::OnTrack:
         return "OnTrack";
-    case EContextStatus::AbortedByLineInitialization:
-        return "AbortedByLineInitialization";
-    case EContextStatus::AbortedByCommitTimeLineCheck:
-        return "AbortedByCommitTimeLineCheck";
     case EContextStatus::AbortedByFailedLockAcquisition:
         return "AbortedByFailedLockAcquisition";
     case EContextStatus::AbortedByLanguage:
diff --git a/Engine/Source/Runtime/Core/Private/AutoRTFM/Memcpy.cpp b/Engine/Source/Runtime/Core/Private/AutoRTFM/Memcpy.cpp
--- a/Engine/Source/Runtime/Core/Private/AutoRTFM/Memcpy.cpp
+++ b/Engine/Source/Runtime/Core/Private/AutoRTFM/Memcpy.cpp
@@ -28,7 +28,7 @@ void* Memcpy(void* InDst, const void* InSrc, size_t Size, FContext* Context)
     }
     FDebug Debug(Context, InDst, InSrc, Size, 0, __FUNCTION__);
     
-    Context->RecordWrite(InDst, Size);
+    Context->RecordWrite(InDst, Size, true);
 
     return memcpy(InDst, InSrc, Size);
 }
@@ -41,7 +41,7 @@ void* Memmove(void* InDst, const void* InSrc, size_t Size, FContext* Context)
     }
     FDebug Debug(Context, InDst, InSrc, Size, 0, __FUNCTION__);
 
-	Context->RecordWrite(InDst, Size);
+	Context->RecordWrite(InDst, Size, true);
 
 	return memmove(InDst, InSrc, Size);
 }
@@ -54,7 +54,7 @@ void* Memset(void* InDst, int Value, size_t Size, FContext* Context)
     }
     FDebug Debug(Context, InDst, nullptr, Size, 0, __FUNCTION__);
 
-	Context->RecordWrite(InDst, Size);
+	Context->RecordWrite(InDst, Size, true);
 
 	return memset(InDst, Value, Size);
 }
diff --git a/Engine/Source/Runtime/Core/Private/AutoRTFM/Transaction.cpp b/Engine/Source/Runtime/Core/Private/AutoRTFM/Transaction.cpp
--- a/Engine/Source/Runtime/Core/Private/AutoRTFM/Transaction.cpp
+++ b/Engine/Source/Runtime/Core/Private/AutoRTFM/Transaction.cpp
@@ -3,6 +3,7 @@
 #if (defined(__AUTORTFM) && __AUTORTFM)
 #include "Transaction.h"
 #include "TransactionInlines.h"
+#include "CallNestInlines.h"
 #include "Debug.h"
 #include "GlobalData.h"
 
@@ -11,7 +12,6 @@ namespace AutoRTFM
 
 FTransaction::FTransaction(FContext* Context)
     : Context(Context)
-    , Parent(Context->GetCurrentTransaction())
 {
 }
 
@@ -30,9 +30,7 @@ void FTransaction::AbortWithoutThrowing()
     {
         fprintf(GetLogFile(), "Aborting (%s)\n", GetContextStatusName(Context->GetStatus()));
     }
-    ASSERT(Context->GetStatus() == EContextStatus::AbortedByLineInitialization
-           || Context->GetStatus() == EContextStatus::AbortedByCommitTimeLineCheck
-           || Context->GetStatus() == EContextStatus::AbortedByFailedLockAcquisition
+    ASSERT(Context->GetStatus() == EContextStatus::AbortedByFailedLockAcquisition
            || Context->GetStatus() == EContextStatus::AbortedByLanguage
            || Context->GetStatus() == EContextStatus::AbortedByRequest);
     ASSERT(Context->GetCurrentTransaction() == this);
@@ -50,7 +48,7 @@ void FTransaction::AbortWithoutThrowing()
 void FTransaction::AbortAndThrow()
 {
     AbortWithoutThrowing();
-    AbortJump.Throw();
+	Context->Throw();
 }
 
 bool FTransaction::AttemptToCommit()
@@ -75,8 +73,9 @@ void FTransaction::Undo()
 {
     for (FWriteLogEntry& Entry : WriteLog)
     {
-        // Skip writes that are into our current transactions stack.
-        if (Context->IsInnerTransactionStack(Entry.OriginalAndSize.Get()))
+        // Skip writes to our current transaction nest if we're scoped. We're about to
+		// leave so the changes don't matter. 
+        if (IsScopedTransaction() && Context->IsInnerTransactionStack(Entry.OriginalAndSize.Get()))
         {
             continue;
         }
@@ -106,8 +105,6 @@ void FTransaction::AbortOuterNest()
 
     switch (Context->GetStatus())
     {
-    case EContextStatus::AbortedByLineInitialization:
-    case EContextStatus::AbortedByCommitTimeLineCheck:
     case EContextStatus::AbortedByFailedLockAcquisition:
         break;
     case EContextStatus::AbortedByRequest:
@@ -131,7 +128,7 @@ void FTransaction::CommitNested()
     for (FWriteLogEntry& Write : WriteLog)
     {
         // Skip writes that are into our current transactions stack.
-        if (Context->IsInnerTransactionStack(Write.OriginalAndSize.Get()))
+        if (IsScopedTransaction() && Context->IsInnerTransactionStack(Write.OriginalAndSize.Get()))
         {
             continue;
         }
diff --git a/Engine/Source/Runtime/Core/Private/AutoRTFM/Transaction.h b/Engine/Source/Runtime/Core/Private/AutoRTFM/Transaction.h
--- a/Engine/Source/Runtime/Core/Private/AutoRTFM/Transaction.h
+++ b/Engine/Source/Runtime/Core/Private/AutoRTFM/Transaction.h
@@ -21,6 +21,7 @@ public:
     
     bool IsNested() const { return !!Parent; }
     FTransaction* GetParent() const { return Parent; }
+	void SetParent(FTransaction* NewParent) { Parent = NewParent; }
 
     // This should just use type displays or ranges. Maybe ranges could even work out great.
     bool IsNestedWithin(const FTransaction* Other) const
@@ -42,19 +43,18 @@ public:
     bool IsDone() const { return bIsDone; }
     void SetIsDone() { bIsDone = true; }
     
+	inline bool IsScopedTransaction() const { return bIsStackScoped; }
+	inline void SetIsScopedTransaction() { bIsStackScoped = true; }
+
     void DeferUntilCommit(std::function<void()>&&);
     void DeferUntilAbort(std::function<void()>&&);
 
-    // Whether this succeeded or not is reflected in Context::GetStatus().
-    template<typename TTryFunctor>
-    void Try(const TTryFunctor& TryFunctor);
-
     void AbortAndThrow();
     void AbortWithoutThrowing();
     bool AttemptToCommit();
 
     // Record that a write is about to occur at the given LogicalAddress of Size bytes.
-    void RecordWrite(void* LogicalAddress, size_t Size);
+    void RecordWrite(void* LogicalAddress, size_t Size, bool bIsClosed);
     void RecordWriteMaxPageSized(void* LogicalAddress, size_t Size);
 
     void DidAllocate(void* LogicalAddress, size_t Size);
@@ -72,15 +72,14 @@ private:
     FContext* Context;
     
     // Are we nested? Then this is the parent.
-    FTransaction* Parent;
+    FTransaction* Parent{nullptr};
 
     // Commit tasks run on commit in forward order. Abort tasks run on abort in reverse order.
     TTaskArray<std::function<void()>> CommitTasks;
     TTaskArray<std::function<void()>> AbortTasks;
 
     bool bIsDone{false};
-
-    FLongJump AbortJump;
+	bool bIsStackScoped{false};
 
     FHitSet HitSet;
     FWriteLog WriteLog;
diff --git a/Engine/Source/Runtime/Core/Private/AutoRTFM/TransactionInlines.h b/Engine/Source/Runtime/Core/Private/AutoRTFM/TransactionInlines.h
--- a/Engine/Source/Runtime/Core/Private/AutoRTFM/TransactionInlines.h
+++ b/Engine/Source/Runtime/Core/Private/AutoRTFM/TransactionInlines.h
@@ -4,7 +4,6 @@
 
 #include "Context.h"
 #include "Transaction.h"
-#include <cstddef>
 
 namespace AutoRTFM
 {
@@ -40,8 +39,23 @@ inline void FTransaction::RecordWriteMaxPageSized(void* LogicalAddress, size_t S
     WriteLog.Push(FWriteLogEntry(LogicalAddress, Size, CopyAddress));
 }
 
-inline void FTransaction::RecordWrite(void* LogicalAddress, size_t Size)
+inline void FTransaction::RecordWrite(void* LogicalAddress, size_t Size, bool bIsClosed)
 {
+	// We don't record any writes to unscoped transactions when the write-address
+	// is within the nearest closed C++ function call nest. This may seem weird, but we
+	// have to prohibit this because we can't track all the comings and goings of
+	// local variables and all child C++ scopes. It's easy to get into a situation
+	// where a rollback will undo the recorded writes to stack-locals that actually
+	// corrupts the undo process itself.
+	if (!IsScopedTransaction()
+		&& !bIsClosed
+		&& Context->IsInnerTransactionStack(LogicalAddress))
+	{
+		// It is an error to write to stack variables within an open nest
+		fprintf(stderr, "FATAL: Writing to local stack memory from an unscoped transaction is not allowed.\n");
+		abort();
+	}
+
     uint8_t* Address = reinterpret_cast<uint8_t*>(LogicalAddress);
 
     size_t I = 0;
@@ -77,20 +91,4 @@ inline void FTransaction::DeferUntilAbort(std::function<void()>&& Callback)
     AbortTasks.Add(std::move(Callback));
 }
 
-template<typename TTryFunctor>
-void FTransaction::Try(const TTryFunctor& TryFunctor)
-{
-    AbortJump.TryCatch(
-        [&] ()
-        {
-            TryFunctor();
-            ASSERT(Context->GetStatus() == EContextStatus::OnTrack);
-        },
-        [&] ()
-        {
-            ASSERT(Context->GetStatus() != EContextStatus::Idle);
-            ASSERT(Context->GetStatus() != EContextStatus::OnTrack);
-        });
-}
-
 } // namespace AutoRTFM
diff --git a/Engine/Source/Runtime/Core/Public/AutoRTFM/AutoRTFM.h b/Engine/Source/Runtime/Core/Public/AutoRTFM/AutoRTFM.h
--- a/Engine/Source/Runtime/Core/Public/AutoRTFM/AutoRTFM.h
+++ b/Engine/Source/Runtime/Core/Public/AutoRTFM/AutoRTFM.h
@@ -42,6 +42,16 @@ typedef enum
     autortfm_committed
 } autortfm_result;
 
+// This must match AutoRTFM::EContextStatus.
+typedef enum
+{
+	autortfm_status_idle,
+	autortfm_status_ontrack,
+	autortfm_status_aborted_by_failed_lock_aquisition,
+	autortfm_status_aborted_by_language,
+	autortfm_status_aborted_by_request
+} autortfm_status;
+
 // Tells if we are currently running in a transaction. This will return true in an
 // open nest (see autortfm_open).
 #if UE_AUTORTFM
@@ -99,17 +109,59 @@ inline void autortfm_commit(void (*work)(void* arg), void* arg)
 }
 #endif
 
-// Request that the current transaction aborts. Only aborts the innermost
-// transaction nest. The corresponding autortfm_transact call will return
-// autortfm_aborted_by_request.
-//
-// Guaranteed to crash if called outside a transaction.
+// Create a new transaction in the open.
 #if UE_AUTORTFM
-void autortfm_abort(void);
+bool autortfm_start_transaction();
 #else
-inline void autortfm_abort(void)
+inline bool autortfm_start_transaction()
+{
+	return false;
+}
+#endif
+
+// End a transaction and commit the changes to be visible to all
+#if UE_AUTORTFM
+autortfm_result autortfm_commit_transaction();
+#else
+inline autortfm_result autortfm_commit_transaction()
+{
+	return autortfm_aborted_by_language;
+}
+#endif
+
+// End a transaction and discard all changes
+#if UE_AUTORTFM
+autortfm_result autortfm_abort_transaction();
+#else
+inline autortfm_result autortfm_abort_transaction()
+{
+	return autortfm_aborted_by_request;
+}
+#endif
+
+// Clear the status of a transaction that was aborted in the open
+#if UE_AUTORTFM
+void autortfm_clear_transaction_status();
+#else
+inline void autortfm_clear_transaction_status()
+{
+}
+#endif
+
+#if UE_AUTORTFM
+bool autortfm_is_aborting();
+#else
+inline bool autortfm_is_aborting()
+{
+	return false;
+}
+#endif
+
+#if UE_AUTORTFM
+bool autortfm_current_nest_throw();
+#else
+inline bool autortfm_current_nest_throw()
 {
-    abort();
 }
 #endif
 
@@ -147,11 +199,24 @@ inline void autortfm_open(void (*work)(void* arg), void* arg)
 //
 // Guaranteed to crash if called outside a transaction.
 #if UE_AUTORTFM
-void autortfm_close(void (*work)(void* arg), void* arg);
+autortfm_status autortfm_close(void (*work)(void* arg), void* arg);
 #else
-inline void autortfm_close(void (*work)(void* arg), void* arg)
+inline autortfm_status autortfm_close(void (*work)(void* arg), void* arg)
 {
     abort();
+	return autortfm_status_aborted_by_language;
+}
+#endif
+
+// Records the pointer and size from the open into the current transaction
+//
+// Guaranteed to assert if called outside a transaction.
+// Guaranteed to assert if called from closed code.
+#if UE_AUTORTFM
+void autortfm_record_open_write(void* Ptr, size_t Size);
+#else
+inline void autortfm_record_open_write(void* Ptr, size_t Size)
+{
 }
 #endif
 
@@ -291,6 +356,15 @@ enum class ETransactionResult
     Committed = autortfm_committed
 };
 
+enum class EContextStatus
+{
+	Idle = autortfm_status_idle,
+	OnTrack = autortfm_status_ontrack,
+	AbortedByFailedLockAcquisition = autortfm_status_aborted_by_failed_lock_aquisition,
+	AbortedByLanguage = autortfm_status_aborted_by_language,
+	AbortedByRequest = autortfm_status_aborted_by_request
+};
+
 inline bool IsTransactional() { return autortfm_is_transactional(); }
 inline bool IsClosed() { return autortfm_is_closed(); }
 
@@ -310,9 +384,93 @@ void Commit(const TFunctor& Functor)
         const_cast<void*>(static_cast<const void*>(&Functor)));
 }
 
-inline void Abort()
+inline bool StartTransaction()
+{
+	return autortfm_start_transaction();
+}
+
+inline ETransactionResult CommitTransaction()
+{
+	return static_cast<ETransactionResult>(autortfm_commit_transaction());
+}
+
+inline ETransactionResult AbortTransaction()
+{
+	return static_cast<ETransactionResult>(autortfm_abort_transaction());
+}
+
+inline void ClearTransactionStatus()
+{
+	autortfm_clear_transaction_status();
+}
+
+inline bool IsAborting()
+{
+	return autortfm_is_aborting();
+}
+
+// RecordOpenWrite records the memory span into the current transaction as written.
+//  If this memory is previously unknown to the transaction, the original value is saved.
+inline void RecordOpenWrite(void* Ptr, size_t Size)
+{
+	autortfm_record_open_write(Ptr, Size);
+}
+
+template<typename TTYPE>
+inline void RecordOpenWrite(TTYPE* Ptr)
 {
-    autortfm_abort();
+	autortfm_record_open_write(Ptr, sizeof(TTYPE));
+}
+
+// RecordOpenRead does nothing right now, but it is intended as a support stub for the day when we move to full AutoSTM
+inline void RecordOpenRead(void const* Ptr, size_t Size)
+{
+}
+
+template<typename TTYPE>
+inline void RecordOpenRead(TTYPE* Ptr)
+{
+	RecordOpenRead(Ptr, sizeof(TTYPE));
+}
+
+// WriteMemory first records the memory span as written (see RecordOpenWrite) and then copies the specified value into it.
+inline void WriteMemory(void* DestPtr, void const* SrcPtr, size_t Size)
+{
+	RecordOpenWrite(DestPtr, Size);
+	memcpy(DestPtr, SrcPtr, Size);
+}
+
+template<typename TTYPE>
+inline void WriteMemoryTrivial(TTYPE* DestPtr, TTYPE const* SrcPtr)
+{
+	RecordOpenWrite(DestPtr, sizeof(TTYPE));
+	*DestPtr = *SrcPtr;
+}
+
+template<typename TTYPE>
+inline void WriteMemory(TTYPE* DestPtr, TTYPE const* SrcPtr)
+{
+	if constexpr (std::is_trivially_copyable<TTYPE>::value)
+	{
+		WriteMemoryTrivial<TTYPE>(DestPtr, SrcPtr);
+	}
+	else
+	{
+		WriteMemory(DestPtr, SrcPtr, sizeof(TTYPE));
+	}
+}
+
+template<typename TTYPE>
+inline void WriteMemory(TTYPE* DestPtr, TTYPE const SrcValue)
+{
+	if constexpr (std::is_trivially_copyable<TTYPE>::value)
+	{
+		WriteMemoryTrivial<TTYPE>(DestPtr, &SrcValue);
+	}
+	else
+	{
+		WriteMemory(DestPtr, &SrcValue, sizeof(TTYPE));
+	}
 }
 
 inline void AbortIfTransactional()
@@ -331,14 +489,19 @@ void Open(const TFunctor& Functor)
     autortfm_open(
         [] (void* Arg) { (*static_cast<const TFunctor*>(Arg))(); },
         const_cast<void*>(static_cast<const void*>(&Functor)));
+
+	if (autortfm_is_aborting())
+	{
+		autortfm_current_nest_throw();
+	}
 }
 
 template<typename TFunctor>
-void Close(const TFunctor& Functor)
+EContextStatus Close(const TFunctor& Functor)
 {
-    autortfm_close(
+    return static_cast<EContextStatus>(autortfm_close(
         [] (void* Arg) { (*static_cast<const TFunctor*>(Arg))(); },
-        const_cast<void*>(static_cast<const void*>(&Functor)));
+        const_cast<void*>(static_cast<const void*>(&Functor))));
 }
 
 inline void RegisterOpenFunction(void* OriginalFunction, void* NewFunction)
