Add a test with a dumb little VM showing how the AutoRTFM calls will work.

diff --git a/Engine/Source/Programs/AutoRTFMTests/Private/API.cpp b/Engine/Source/Programs/AutoRTFMTests/Private/API.cpp
--- a/Engine/Source/Programs/AutoRTFMTests/Private/API.cpp
+++ b/Engine/Source/Programs/AutoRTFMTests/Private/API.cpp
@@ -111,7 +111,7 @@ TEST_CASE("API.autortfm_commit")
     REQUIRE(42 == Answer);
 }
 
-TEST_CASE("API.autortfm_abort")
+TEST_CASE("API.autortfm_abort_transaction")
 {
     bool BeforeNest = false;
     bool InNest = false;
@@ -244,10 +244,10 @@ TEST_CASE("API.autortfm_close")
         AutoRTFM::Open([&]
         {
             // A closed call inside an open does not abort either.
-            autortfm_close([](void* const Arg)
+			REQUIRE(autortfm_status_ontrack == autortfm_close([](void* const Arg)
             {
                 *static_cast<bool* const>(Arg) = true;
-            }, &InClosedNestInOpenNest);
+            }, &InClosedNestInOpenNest));
 
             InOpenNest = true;
         });
@@ -839,10 +839,10 @@ TEST_CASE("API.Close")
         AutoRTFM::Open([&]
         {
             // A closed call inside an open does not abort either.
-            AutoRTFM::Close([&]
+			REQUIRE(AutoRTFM::EContextStatus::OnTrack == AutoRTFM::Close([&]
             {
                 InClosedNestInOpenNest = true;
-            });
+            }));
 
             InOpenNest = true;
         });
diff --git a/Engine/Source/Programs/AutoRTFMTests/Private/FakeVM.cpp b/Engine/Source/Programs/AutoRTFMTests/Private/FakeVM.cpp
new file mode 100644
--- /dev/null
+++ b/Engine/Source/Programs/AutoRTFMTests/Private/FakeVM.cpp
@@ -0,0 +1,188 @@
+// Copyright Epic Games, Inc. All Rights Reserved.
+
+#pragma autortfm
+
+#include <catch_amalgamated.hpp>
+#include <AutoRTFM/AutoRTFM.h>
+#include <vector>
+
+enum class EOpcodeType
+{
+	StartTransaction,
+	CommitTransaction,
+	Call,
+	WriteToValue
+};
+
+typedef void (*FCall)(int*, int);
+
+struct FOpcode final
+{
+private:
+	explicit FOpcode(EOpcodeType Type, int Value = -1, FCall Call = nullptr) : Type(Type), Value(Value), Call(Call) {}
+
+public:
+	EOpcodeType Type;
+	int Value;
+	FCall Call;
+
+	static FOpcode CreateStartTransaction(int RecoveryOpcodeOffset)
+	{
+		return FOpcode(EOpcodeType::StartTransaction, RecoveryOpcodeOffset);
+	}
+
+	static FOpcode CreateCommitTransaction()
+	{
+		return FOpcode(EOpcodeType::CommitTransaction);
+	}
+
+	static FOpcode CreateCall(int Arg, FCall Call)
+	{
+		return FOpcode(EOpcodeType::Call, Arg, Call);
+	}
+
+	static FOpcode CreateWriteToValue(int Arg)
+	{
+		return FOpcode(EOpcodeType::WriteToValue, Arg);	
+	}
+};
+
+void Good(int* Value, int Arg)
+{
+	*Value += Arg;
+}
+
+void Bad(int*, int)
+{
+	AutoRTFM::AbortTransaction();
+}
+
+AutoRTFM::ETransactionResult FakeVM(const FOpcode* Opcodes, const int Length, int& Value)
+{
+	return AutoRTFM::Transact([Opcodes, Length, &Value]
+	{
+		AutoRTFM::Open([Opcodes, Length, &Value]
+		{
+			bool ShouldThrowNest = false;
+
+			// A scope - just to ensure the destructor of `RecoveryOpcodeStack` has been run on throw!
+			{
+				std::vector<int> RecoveryOpcodeStack;
+
+				for (int i = 0; i < Length; i++)
+				{
+					switch (Opcodes[i].Type)
+					{
+					default:
+						FAIL("Unhandled opcode kind!");
+						break;
+					case EOpcodeType::StartTransaction:
+						AutoRTFM::StartTransaction();
+						RecoveryOpcodeStack.push_back(Opcodes[i].Value);
+						break;
+					case EOpcodeType::CommitTransaction:
+						REQUIRE(AutoRTFM::ETransactionResult::Committed == AutoRTFM::CommitTransaction());
+						RecoveryOpcodeStack.pop_back();
+						break;
+					case EOpcodeType::Call:
+					{
+						const FCall Call = Opcodes[i].Call;
+						const int Arg = Opcodes[i].Value;
+
+						const AutoRTFM::EContextStatus Status = AutoRTFM::Close([&Value, Arg, Call]
+							{
+								Call(&Value, Arg);
+							});
+
+						if (AutoRTFM::EContextStatus::OnTrack != Status)
+						{
+							if (RecoveryOpcodeStack.empty())
+							{
+								ShouldThrowNest = true;
+								break;
+							}
+
+							// We've handled the bad transaction status here, so clear it out!
+							AutoRTFM::ClearTransactionStatus();
+
+							// -1 just cause the loop will i++!
+							i = RecoveryOpcodeStack.back() - 1;
+							RecoveryOpcodeStack.pop_back();
+						}
+						break;
+					}
+					case EOpcodeType::WriteToValue:
+					{
+						const int Arg = Value - Opcodes[i].Value;
+						AutoRTFM::WriteMemory(&Value, &Arg);
+						break;
+					}
+					}
+
+					if (ShouldThrowNest)
+					{
+						break;
+					}
+				}
+			}
+
+			if (ShouldThrowNest)
+			{
+				AutoRTFM::CurrentNestThrow();
+			}
+		});
+	});
+}
+
+TEST_CASE("FakeVM.NativeCallInTransaction")
+{
+	constexpr int Length = 6;
+	const FOpcode Opcodes[Length] =
+	{
+		// 5 being the offset just after the commit
+		FOpcode::CreateStartTransaction(5),
+		FOpcode::CreateWriteToValue(42),
+		FOpcode::CreateCall(43, &Good),
+		FOpcode::CreateWriteToValue(13),
+		FOpcode::CreateCommitTransaction(),
+		FOpcode::CreateWriteToValue(-53),
+	};
+
+	int Value = 100;
+	REQUIRE(AutoRTFM::ETransactionResult::Committed == FakeVM(Opcodes, Length, Value));
+	REQUIRE(Value == (100 - 42 + 43 - 13 - (-53)));
+}
+
+TEST_CASE("FakeVM.AbortInTransaction")
+{
+	constexpr int Length = 6;
+	const FOpcode Opcodes[Length] =
+	{
+		// 5 being the offset just after the commit
+		FOpcode::CreateStartTransaction(5),
+		FOpcode::CreateWriteToValue(42),
+		FOpcode::CreateCall(43, &Bad),
+		FOpcode::CreateWriteToValue(13),
+		FOpcode::CreateCommitTransaction(),
+		FOpcode::CreateWriteToValue(-53),
+	};
+
+	int Value = 100;
+	REQUIRE(AutoRTFM::ETransactionResult::Committed == FakeVM(Opcodes, Length, Value));
+	REQUIRE(Value == (100 - (-53)));
+}
+
+TEST_CASE("FakeVM.Abort")
+{
+	constexpr int Length = 3;
+	const FOpcode Opcodes[Length] =
+	{
+		FOpcode::CreateWriteToValue(42),
+		FOpcode::CreateCall(43, &Bad),
+		FOpcode::CreateWriteToValue(13),
+	};
+
+	int Value = 100;
+	REQUIRE(AutoRTFM::ETransactionResult::AbortedByRequest == FakeVM(Opcodes, Length, Value));
+	REQUIRE(Value == 100);
+}
diff --git a/Engine/Source/Programs/AutoRTFMTests/Private/OpenAPI.cpp b/Engine/Source/Programs/AutoRTFMTests/Private/OpenAPI.cpp
--- a/Engine/Source/Programs/AutoRTFMTests/Private/OpenAPI.cpp
+++ b/Engine/Source/Programs/AutoRTFMTests/Private/OpenAPI.cpp
@@ -30,6 +30,8 @@ TEST_CASE("OpenAPI.StartAbortAndStartAgain")
 			AutoRTFM::RecordOpenWrite(&valueC);
 			valueC = 30;
 			AutoRTFM::AbortTransaction();
+
+			AutoRTFM::CurrentNestThrow();
 		});
 	});
 
@@ -85,6 +87,7 @@ TEST_CASE("OpenAPI.AbortTransactionScopedFromOpen")
 	{
 		AutoRTFM::Open([&]() {
 			AutoRTFM::AbortTransaction();
+			AutoRTFM::CurrentNestThrow();
 		});
 		FAIL("AutoRTFM::Open failed to throw after an abort");
 	});
@@ -115,6 +118,7 @@ TEST_CASE("OpenAPI.AbortTransactionDoubleScopedFromOpen")
 			value = 42424242;
 			AutoRTFM::Open([&]() {
 				AutoRTFM::AbortTransaction();
+				AutoRTFM::CurrentNestThrow();
 			});
 			FAIL("AutoRTFM::Open failed to throw after an abort");
 			value = 24242424;
@@ -372,6 +376,7 @@ TEST_CASE("OpenAPI.Commit_TransactOpenCloseAbort")
 			});
 
 			AutoRTFM::AbortTransaction(); // undoes value = 42 in the open
+			AutoRTFM::CurrentNestThrow();
 			REQUIRE(valueLocal == 1.0);
 		});
 
@@ -499,6 +504,7 @@ TEST_CASE("OpenAPI.StackWriteAbortInTheOpen1")
 			value = 10;
 			AutoRTFM::AbortTransaction();
 			REQUIRE(value == 0);
+			AutoRTFM::CurrentNestThrow();
 		});
 	});
 }
@@ -550,6 +556,7 @@ TEST_CASE("OpenAPI.StackWriteAbortInTheOpen2")
 			AutoRTFM::WriteMemory(&value, 10); // Illegal to write to value because it's in the inner-most closed-nest
 
 			AutoRTFM::AbortTransaction();
+			AutoRTFM::CurrentNestThrow();
 		});
 
 		// Never gets here
@@ -602,6 +609,7 @@ TEST_CASE("OpenAPI.WriteTrivialStructure")
 			REQUIRE(data.E[2] == 345);
 			REQUIRE(data.E[3] == 456);
 			REQUIRE(data.E[4] == 567);
+			AutoRTFM::CurrentNestThrow();
 		});
 	});
 }
@@ -659,6 +667,7 @@ TEST_CASE("OpenAPI.WriteTrivialStructure2")
 				REQUIRE(data.E[2] == 345);
 				REQUIRE(data.E[3] == 456);
 				REQUIRE(data.E[4] == 567);
+				AutoRTFM::CurrentNestThrow();
 			});
 	});
 }
@@ -738,7 +747,8 @@ TEST_CASE("OpenAPI.StartCloseOpenCommit")
 	AutoRTFM::StartTransaction();
 
 	// Can't close outside of a transaction
-	AutoRTFM::Close([&]() {
+	AutoRTFM::Close([&]()
+	{
 		value = 420;
 	});
 
@@ -796,6 +806,7 @@ TEST_CASE("OpenAPI.TransOpenStartCloseAbortAbort")
 			AutoRTFM::AbortTransaction();
 			bGetsToC = true;
 			REQUIRE(value == 42);
+			AutoRTFM::CurrentNestThrow();
 		});
 
 		bGetsToD = true;
@@ -819,8 +830,10 @@ TEST_CASE("OpenAPI.TransOpenTransCloseAbortAbort")
 
 	AutoRTFM::Transact([&]() 
 	{
-		AutoRTFM::Open([&]() {
-			AutoRTFM::Transact([&](){
+		AutoRTFM::Open([&]()
+		{
+			AutoRTFM::Transact([&]()
+			{
 				int value = 10;
 				value++;
 
@@ -852,4 +865,3 @@ TEST_CASE("OpenAPI.TransOpenTransCloseAbortAbort")
 	REQUIRE(bGetsToD == true);
 	REQUIRE(!AutoRTFM::IsTransactional());
 }
-
diff --git a/Engine/Source/Runtime/Core/Private/AutoRTFM/API.cpp b/Engine/Source/Runtime/Core/Private/AutoRTFM/API.cpp
--- a/Engine/Source/Runtime/Core/Private/AutoRTFM/API.cpp
+++ b/Engine/Source/Runtime/Core/Private/AutoRTFM/API.cpp
@@ -331,7 +331,7 @@ UE_AUTORTFM_REGISTER_OPEN_FUNCTION(autortfm_start_transaction);
 
 void STM_autortfm_commit_transaction(FContext* Context)
 {
-	fprintf(stderr, "autostm_committransaction called from closed code.\n");
+	fprintf(stderr, "autostm_commit_transaction called from closed code.\n");
 	abort();
 }
 UE_AUTORTFM_REGISTER_OPEN_FUNCTION(autortfm_commit_transaction);
diff --git a/Engine/Source/Runtime/Core/Public/AutoRTFM/AutoRTFM.h b/Engine/Source/Runtime/Core/Public/AutoRTFM/AutoRTFM.h
--- a/Engine/Source/Runtime/Core/Public/AutoRTFM/AutoRTFM.h
+++ b/Engine/Source/Runtime/Core/Public/AutoRTFM/AutoRTFM.h
@@ -201,7 +201,7 @@ inline void autortfm_open(void (*work)(void* arg), void* arg)
 //
 // Guaranteed to crash if called outside a transaction.
 #if UE_AUTORTFM
-autortfm_status autortfm_close(void (*work)(void* arg), void* arg);
+[[nodiscard]] autortfm_status autortfm_close(void (*work)(void* arg), void* arg);
 #else
 inline autortfm_status autortfm_close(void (*work)(void* arg), void* arg)
 {
@@ -403,6 +403,11 @@ inline bool IsAborting()
 	return autortfm_is_aborting();
 }
 
+inline bool CurrentNestThrow()
+{
+	return autortfm_current_nest_throw();
+}
+
 // RecordOpenWrite records the memory span into the current transaction as written.
 //  If this memory is previously unknown to the transaction, the original value is saved.
 inline void RecordOpenWrite(void* Ptr, size_t Size)
@@ -483,15 +488,10 @@ void Open(const TFunctor& Functor)
     autortfm_open(
         [] (void* Arg) { (*static_cast<const TFunctor*>(Arg))(); },
         const_cast<void*>(static_cast<const void*>(&Functor)));
-
-	if (autortfm_is_aborting())
-	{
-		autortfm_current_nest_throw();
-	}
 }
 
 template<typename TFunctor>
-EContextStatus Close(const TFunctor& Functor)
+[[nodiscard]] EContextStatus Close(const TFunctor& Functor)
 {
     return static_cast<EContextStatus>(autortfm_close(
         [] (void* Arg) { (*static_cast<const TFunctor*>(Arg))(); },
