Fix nodiscard warnings for clang

diff --git a/Engine/Source/Programs/AutoRTFMTests/Private/API.cpp b/Engine/Source/Programs/AutoRTFMTests/Private/API.cpp
--- a/Engine/Source/Programs/AutoRTFMTests/Private/API.cpp
+++ b/Engine/Source/Programs/AutoRTFMTests/Private/API.cpp
@@ -75,10 +75,10 @@ TEST_CASE("API.autortfm_is_closed")
         {
             InOpenNest = autortfm_is_closed();
 
-            AutoRTFM::Close([&]
+            REQUIRE(AutoRTFM::EContextStatus::OnTrack == AutoRTFM::Close([&]
             {
                 InClosedNestInOpenNest = autortfm_is_closed();
-            });
+            }));
         });
     });
 
@@ -236,10 +236,10 @@ TEST_CASE("API.autortfm_close")
     REQUIRE(AutoRTFM::ETransactionResult::AbortedByRequest == AutoRTFM::Transact([&]
     {
         // A closed call inside a transaction does not abort.
-        autortfm_close([](void* const Arg)
+        REQUIRE(autortfm_status_ontrack == autortfm_close([](void* const Arg)
         {
             *static_cast<bool* const>(Arg) = true;
-        }, &InClosedNest);
+        }, &InClosedNest));
 
         AutoRTFM::Open([&]
         {
@@ -670,10 +670,10 @@ TEST_CASE("API.IsClosed")
         {
             InOpenNest = AutoRTFM::IsClosed();
 
-            AutoRTFM::Close([&]
+            REQUIRE(AutoRTFM::EContextStatus::OnTrack == AutoRTFM::Close([&]
             {
                 InClosedNestInOpenNest = AutoRTFM::IsClosed();
-            });
+            }));
         });
     });
 
@@ -831,10 +831,10 @@ TEST_CASE("API.Close")
     REQUIRE(AutoRTFM::ETransactionResult::AbortedByRequest == AutoRTFM::Transact([&]
     {
         // A closed call inside a transaction does not abort.
-        AutoRTFM::Close([&]
+        REQUIRE(AutoRTFM::EContextStatus::OnTrack == AutoRTFM::Close([&]
         {
             InClosedNest = true;
-        });
+        }));
 
         AutoRTFM::Open([&]
         {
diff --git a/Engine/Source/Programs/AutoRTFMTests/Private/OpenAPI.cpp b/Engine/Source/Programs/AutoRTFMTests/Private/OpenAPI.cpp
--- a/Engine/Source/Programs/AutoRTFMTests/Private/OpenAPI.cpp
+++ b/Engine/Source/Programs/AutoRTFMTests/Private/OpenAPI.cpp
@@ -56,10 +56,10 @@ TEST_CASE("OpenAPI.RecordDataClosed_Illegal", "[.]")
 	int value = 0;
 	AutoRTFM::ETransactionResult transactResult = AutoRTFM::Transact([&]()
 	{
-		AutoRTFM::Close([&]() {
+		REQUIRE(AutoRTFM::EContextStatus::OnTrack == AutoRTFM::Close([&]() {
 			AutoRTFM::RecordOpenWrite(&value); // Illegal. Can't record writes explicitly while closed
 			value = 1;
-		});
+		}));
 	});
 
 	REQUIRE(transactResult == AutoRTFM::ETransactionResult::Committed);
@@ -98,9 +98,9 @@ TEST_CASE("OpenAPI.AbortTransactionScopedFromClosed")
 {
 	auto transactResult = AutoRTFM::Transact([&]()
 	{
-		AutoRTFM::Close([&]() {
+		REQUIRE(AutoRTFM::EContextStatus::OnTrack == AutoRTFM::Close([&]() {
 			AutoRTFM::AbortTransaction();
-		});
+		}));
 		FAIL("AutoRTFM::Close should have no-op'ed because it's already closed from the Transact");
 	});
 	REQUIRE(transactResult == AutoRTFM::ETransactionResult::AbortedByRequest);
@@ -242,7 +242,7 @@ TEST_CASE("OpenAPI.OpenCloseOpenClose")
 			REQUIRE(!AutoRTFM::IsClosed());
 
 			// -------------------------------------
-			AutoRTFM::Close([&]() {
+			REQUIRE(AutoRTFM::EContextStatus::OnTrack == AutoRTFM::Close([&]() {
 				// C - WE ARE CLOSED AGAIN
 				if (!AutoRTFM::IsClosed()) FAIL("C - NOT CLOSED AS EXPECTED!");
 
@@ -279,7 +279,7 @@ TEST_CASE("OpenAPI.OpenCloseOpenClose")
 				// G - BACK TO CLOSED AGAIN
 				if (!AutoRTFM::IsClosed()) FAIL("NOT CLOSED!");
 
-			});
+			}));
 			// -------------------------------------
 			// H - BACK TO OPEN 
 			REQUIRE(!AutoRTFM::IsClosed());
@@ -330,9 +330,9 @@ TEST_CASE("OpenAPI.Commit_TransactOpenCloseCommit")
 		AutoRTFM::Open([&]() {
 			AutoRTFM::StartTransaction();
 
-			AutoRTFM::Close([&]() {
+			REQUIRE(AutoRTFM::EContextStatus::OnTrack == AutoRTFM::Close([&]() {
 				value = 42;
-			});
+			}));
 
 			REQUIRE(value == 42); // RTFM writes through immediately, so we can see this value in the open
 			AutoRTFM::CommitTransaction();
@@ -370,10 +370,10 @@ TEST_CASE("OpenAPI.Commit_TransactOpenCloseAbort")
 			AutoRTFM::StartTransaction();
 
 			// Closing from the open doesn't work
-			AutoRTFM::Close([&]() {
+			REQUIRE(AutoRTFM::EContextStatus::OnTrack == AutoRTFM::Close([&]() {
 				value = 42;
 				valueLocal = 10.0;
-			});
+			}));
 
 			AutoRTFM::AbortTransaction(); // undoes value = 42 in the open
 			AutoRTFM::CurrentNestThrow();
@@ -478,13 +478,13 @@ TEST_CASE("OpenAPI.StackWriteCommitInTheOpen2")
 		AutoRTFM::Open([&]()
 		{
 			AutoRTFM::StartTransaction();
-			AutoRTFM::Close([&]()
+			REQUIRE(AutoRTFM::EContextStatus::OnTrack == AutoRTFM::Close([&]()
 			{
 				AutoRTFM::Open([&]() {
 					AutoRTFM::RecordOpenWrite(&value);
 					value = 10;
 				});
-			});
+			}));
 
 			AutoRTFM::CommitTransaction();
 			REQUIRE(value == 10);
@@ -680,7 +680,7 @@ TEST_CASE("OpenAPI.Footgun1")
 	AutoRTFM::Transact([&]()
 	{
 		// Does nothing - already closed
-		AutoRTFM::Close([&]()
+		REQUIRE(AutoRTFM::EContextStatus::OnTrack == AutoRTFM::Close([&]()
 		{
 			// Recorded valueB as starting at 0.
 			valueB = 123;
@@ -695,7 +695,7 @@ TEST_CASE("OpenAPI.Footgun1")
 			// valueA is now recorded as starting at 10
 			valueA = 20;
 			AutoRTFM::AbortTransaction();
-		});
+		}));
 	});
 
 	// We rollback the transaction to the value we had when we first recorded the address
@@ -710,7 +710,7 @@ TEST_CASE("OpenAPI.Footgun2")
 	AutoRTFM::Transact([&]()
 	{
 		// Does nothing - already closed
-		AutoRTFM::Close([&]()
+		REQUIRE(AutoRTFM::EContextStatus::OnTrack == AutoRTFM::Close([&]()
 		{
 			// Recorded valueB as starting at 0.
 			valueB = 20;
@@ -727,7 +727,7 @@ TEST_CASE("OpenAPI.Footgun2")
 			// valueA is now recorded as starting at 10
 			valueC = 40;
 			AutoRTFM::AbortTransaction();
-		});
+		}));
 	});
 
 	// We rollback the transaction to the value we had when we first recorded the address
@@ -747,10 +747,10 @@ TEST_CASE("OpenAPI.StartCloseOpenCommit")
 	AutoRTFM::StartTransaction();
 
 	// Can't close outside of a transaction
-	AutoRTFM::Close([&]()
+	REQUIRE(AutoRTFM::EContextStatus::OnTrack == AutoRTFM::Close([&]()
 	{
 		value = 420;
-	});
+	}));
 
 	// assignment within the close should be visible to us
 	REQUIRE(value == 420);
@@ -839,12 +839,12 @@ TEST_CASE("OpenAPI.TransOpenTransCloseAbortAbort")
 
 				value = 42;
 				// Can't close outside of a Transact
-				AutoRTFM::Close([&]()
+				REQUIRE(AutoRTFM::EContextStatus::OnTrack == AutoRTFM::Close([&]()
 				{
 					value = 420;
 					AutoRTFM::AbortTransaction();
 					bGetsToA = true;
-				});
+				}));
 
 				REQUIRE(bGetsToA == false);
 
diff --git a/Engine/Source/Runtime/Core/Private/AutoRTFM/StdLibSupport.cpp b/Engine/Source/Runtime/Core/Private/AutoRTFM/StdLibSupport.cpp
--- a/Engine/Source/Runtime/Core/Private/AutoRTFM/StdLibSupport.cpp
+++ b/Engine/Source/Runtime/Core/Private/AutoRTFM/StdLibSupport.cpp
@@ -78,11 +78,7 @@ void* STM_realloc(void* Ptr, size_t Size, FContext* Context)
 UE_AUTORTFM_REGISTER_OPEN_FUNCTION(realloc);
 
 
-size_t STM_strlen(const char* Str, FContext* Context)
-{
-	AutoRTFM::Unreachable();
-}
-UE_AUTORTFM_REGISTER_OPEN_FUNCTION(strlen);
+UE_AUTORTFM_REGISTER_SELF_FUNCTION(strlen);
 
 // FIXME: This is only correct when:
 // - Str is newly allocated
