Fix compiler errors with -nopch, to include Catch you need to use Catch2Includes to prevent catch headers for including windows to before our overrides

diff --git a/Engine/Source/Programs/AutoRTFMTests/Private/Allocators.cpp b/Engine/Source/Programs/AutoRTFMTests/Private/Allocators.cpp
--- a/Engine/Source/Programs/AutoRTFMTests/Private/Allocators.cpp
+++ b/Engine/Source/Programs/AutoRTFMTests/Private/Allocators.cpp
@@ -2,7 +2,7 @@
 
 #pragma autortfm
 
-#include <catch_amalgamated.hpp>
+#include "Catch2Includes.h"
 #include <AutoRTFM/AutoRTFM.h>
 #include <new>
 
diff --git a/Engine/Source/Programs/AutoRTFMTests/Private/FakeVM.cpp b/Engine/Source/Programs/AutoRTFMTests/Private/FakeVM.cpp
--- a/Engine/Source/Programs/AutoRTFMTests/Private/FakeVM.cpp
+++ b/Engine/Source/Programs/AutoRTFMTests/Private/FakeVM.cpp
@@ -2,7 +2,7 @@
 
 #pragma autortfm
 
-#include <catch_amalgamated.hpp>
+#include "Catch2Includes.h"
 #include <AutoRTFM/AutoRTFM.h>
 #include <vector>
 
