Add tests for a bunch of cstring functions in AutoRTFM.

diff --git a/Engine/Source/Programs/AutoRTFMTests/Private/CString.cpp b/Engine/Source/Programs/AutoRTFMTests/Private/CString.cpp
new file mode 100644
--- /dev/null
+++ b/Engine/Source/Programs/AutoRTFMTests/Private/CString.cpp
@@ -0,0 +1,208 @@
+// Copyright Epic Games, Inc. All Rights Reserved.
+
+#pragma autortfm
+
+#include "Catch2Includes.h"
+#include <AutoRTFM/AutoRTFM.h>
+#include <cstring>
+
+TEST_CASE("CString.memcpy")
+{
+	constexpr unsigned Size = 128;
+	const char* const From = "Hello, world!";
+	char To[Size];
+	memset(To, 0, Size);
+	AutoRTFM::Commit([&]
+	{
+		memcpy(To, From, 5);
+	});
+
+	REQUIRE(0 == strcmp("Hello", To));
+}
+
+TEST_CASE("CString.memmove")
+{
+	constexpr unsigned Size = 128;
+	char To[Size] = "Hello, world!";
+	AutoRTFM::Commit([&]
+	{
+		memmove(To + 7, To, 5);
+	});
+
+	REQUIRE(0 == strcmp("Hello, Hello!", To));
+}
+
+TEST_CASE("CString.strcpy")
+{
+	constexpr unsigned Size = 128;
+	const char* const From = "Hello, world!";
+	char To[Size];
+
+#pragma clang optimize off
+	AutoRTFM::Commit([&]
+		{
+			strcpy(To, From);
+		});
+#pragma clang optimize on
+
+	REQUIRE(0 == strcmp(From, To));
+}
+
+TEST_CASE("CString.strncpy")
+{
+	constexpr unsigned Size = 128;
+	const char* const From = "Hello, world!";
+	char To[Size];
+	memset(To, 0, Size);
+
+#pragma clang optimize off
+	AutoRTFM::Commit([&]
+		{
+			strncpy(To, From, 5);
+		});
+#pragma clang optimize on
+
+	REQUIRE(0 == strcmp("Hello", To));
+}
+
+TEST_CASE("CString.strcat")
+{
+	constexpr unsigned Size = 128;
+	char To[Size] = "Hello";
+
+#pragma clang optimize off
+	AutoRTFM::Commit([&]
+		{
+			strcat(To, ", world!");
+		});
+#pragma clang optimize on
+
+	REQUIRE(0 == strcmp("Hello, world!", To));
+}
+
+TEST_CASE("CString.strncat")
+{
+	constexpr unsigned Size = 128;
+	char To[Size] = "Hello";
+
+#pragma clang optimize off
+	AutoRTFM::Commit([&]
+		{
+			strncat(To, ", world! Not this!", 8);
+		});
+#pragma clang optimize on
+
+	REQUIRE(0 == strcmp("Hello, world!", To));
+}
+
+TEST_CASE("CString.memcmp")
+{
+	constexpr unsigned Size = 128;
+	char A[Size] = "This";
+
+	int Compare;
+#pragma clang optimize off
+	AutoRTFM::Commit([&]
+		{
+			Compare = memcmp(A, "That", 4);
+		});
+#pragma clang optimize on
+
+	REQUIRE(0 < Compare);
+}
+
+TEST_CASE("CString.strcmp")
+{
+	constexpr unsigned Size = 128;
+	char A[Size] = "This";
+
+	int Compare;
+#pragma clang optimize off
+	AutoRTFM::Commit([&]
+		{
+			Compare = strcmp(A, "That");
+		});
+#pragma clang optimize on
+
+	REQUIRE(0 < Compare);
+}
+
+TEST_CASE("CString.strncmp")
+{
+	constexpr unsigned Size = 128;
+	char A[Size] = "This";
+
+	int Compare;
+#pragma clang optimize off
+	AutoRTFM::Commit([&]
+		{
+			Compare = strncmp(A, "That", 3);
+		});
+#pragma clang optimize on
+
+	REQUIRE(0 < Compare);
+}
+
+TEST_CASE("CString.strchr")
+{
+	constexpr unsigned Size = 128;
+	char A[Size] = "This";
+
+	const char* Result;
+#pragma clang optimize off
+	AutoRTFM::Commit([&]
+		{
+			Result = strchr(A, 'i');
+		});
+#pragma clang optimize on
+
+	REQUIRE((A + 2) == Result);
+}
+
+TEST_CASE("CString.strrchr")
+{
+	constexpr unsigned Size = 128;
+	char A[Size] = "This";
+
+	const char* Result;
+#pragma clang optimize off
+	AutoRTFM::Commit([&]
+		{
+			Result = strrchr(A, 'i');
+		});
+#pragma clang optimize on
+
+	REQUIRE((A + 2) == Result);
+}
+
+TEST_CASE("CString.strstr")
+{
+	constexpr unsigned Size = 128;
+	char A[Size] = "This";
+
+	const char* Result;
+#pragma clang optimize off
+	AutoRTFM::Commit([&]
+		{
+			Result = strstr(A, "is");
+		});
+#pragma clang optimize on
+
+	REQUIRE((A + 2) == Result);
+}
+
+TEST_CASE("CString.strlen")
+{
+	constexpr unsigned Size = 128;
+	char A[Size] = "This";
+
+	size_t Result;
+#pragma clang optimize off
+	AutoRTFM::Commit([&]
+		{
+			Result = strlen(A);
+		});
+#pragma clang optimize on
+
+	REQUIRE(4 == Result);
+}
diff --git a/Engine/Source/Runtime/Core/Private/AutoRTFM/StdLibSupport.cpp b/Engine/Source/Runtime/Core/Private/AutoRTFM/StdLibSupport.cpp
--- a/Engine/Source/Runtime/Core/Private/AutoRTFM/StdLibSupport.cpp
+++ b/Engine/Source/Runtime/Core/Private/AutoRTFM/StdLibSupport.cpp
@@ -77,7 +77,47 @@ void* STM_realloc(void* Ptr, size_t Size, FContext* Context)
 }
 UE_AUTORTFM_REGISTER_OPEN_FUNCTION(realloc);
 
+char* STM_strcpy(char* const Dst, const char* const Src, FContext* const Context)
+{
+    const size_t SrcLen = strlen(Src);
+
+    Context->RecordWrite(Dst, SrcLen, true);
+    return strcpy(Dst, Src);
+}
+UE_AUTORTFM_REGISTER_OPEN_FUNCTION(strcpy);
+
+char* STM_strncpy(char* const Dst, const char* const Src, const size_t Num, FContext* const Context)
+{
+    Context->RecordWrite(Dst, Num, true);
+    return strncpy(Dst, Src, Num);
+}
+UE_AUTORTFM_REGISTER_OPEN_FUNCTION(strncpy);
+
+char* STM_strcat(char* const Dst, const char* const Src, FContext* const Context)
+{
+    const size_t DstLen = strlen(Dst);
+    const size_t SrcLen = strlen(Src);
+
+    Context->RecordWrite(Dst + DstLen, SrcLen + 1, true);
+    return strcat(Dst, Src);
+}
+UE_AUTORTFM_REGISTER_OPEN_FUNCTION(strcat);
+
+char* STM_strncat(char* const Dst, const char* const Src, const size_t Num, FContext* const Context)
+{
+    const size_t DstLen = strlen(Dst);
 
+    Context->RecordWrite(Dst + DstLen, Num + 1, true);
+    return strncat(Dst, Src, Num);
+}
+UE_AUTORTFM_REGISTER_OPEN_FUNCTION(strncat);
+
+UE_AUTORTFM_REGISTER_SELF_FUNCTION(memcmp);
+UE_AUTORTFM_REGISTER_SELF_FUNCTION(strcmp);
+UE_AUTORTFM_REGISTER_SELF_FUNCTION(strncmp);
+UE_AUTORTFM_REGISTER_SELF_FUNCTION(static_cast<const char*(*)(const char*, int)>(&strchr));
+UE_AUTORTFM_REGISTER_SELF_FUNCTION(static_cast<const char* (*)(const char*, int)>(&strrchr));
+UE_AUTORTFM_REGISTER_SELF_FUNCTION(static_cast<const char* (*)(const char*, const char*)>(&strstr));
 UE_AUTORTFM_REGISTER_SELF_FUNCTION(strlen);
 
 // FIXME: This is only correct when:
@@ -94,26 +134,6 @@ int STM_snprintf(char* Str, size_t Size, char* Format, FContext* Context, ...)
 }
 UE_AUTORTFM_REGISTER_OPEN_FUNCTION(snprintf);
 
-char *STM_strncpy(char* Dst, const char* Src, size_t Len, FContext* Context)
-{
-    AutoRTFM::Unreachable();
-}
-
-#ifdef _MSC_VER
-/*
-   Disable warning about deprecated STD C functions.
-*/
-#pragma warning(disable : 4996)
-
-#pragma warning(push)
-#endif
-
-UE_AUTORTFM_REGISTER_OPEN_FUNCTION(strncpy);
-
-#ifdef _MSC_VER
-#pragma warning(pop)
-#endif
-
 int STM_printf(const char* Format, FContext* Context, ...)
 {
     va_list ArgList;
