Make the AutoRTFM LLVM compiler pass not elide transactional calls to AutoRTFM::Open, and only elide the non-transactional ones.

diff --git a/Engine/Source/Programs/AutoRTFMTests/Private/FakeVM.cpp b/Engine/Source/Programs/AutoRTFMTests/Private/FakeVM.cpp
--- a/Engine/Source/Programs/AutoRTFMTests/Private/FakeVM.cpp
+++ b/Engine/Source/Programs/AutoRTFMTests/Private/FakeVM.cpp
@@ -63,72 +63,57 @@ AutoRTFM::ETransactionResult FakeVM(const FOpcode* Opcodes, const int Length, in
 	{
 		AutoRTFM::Open([Opcodes, Length, &Value]
 		{
-			bool ShouldThrowNest = false;
+			std::vector<int> RecoveryOpcodeStack;
 
-			// A scope - just to ensure the destructor of `RecoveryOpcodeStack` has been run on throw!
+			for (int i = 0; i < Length; i++)
 			{
-				std::vector<int> RecoveryOpcodeStack;
-
-				for (int i = 0; i < Length; i++)
+				switch (Opcodes[i].Type)
 				{
-					switch (Opcodes[i].Type)
-					{
-					default:
-						FAIL("Unhandled opcode kind!");
-						break;
-					case EOpcodeType::StartTransaction:
-						AutoRTFM::StartTransaction();
-						RecoveryOpcodeStack.push_back(Opcodes[i].Value);
-						break;
-					case EOpcodeType::CommitTransaction:
-						REQUIRE(AutoRTFM::ETransactionResult::Committed == AutoRTFM::CommitTransaction());
-						RecoveryOpcodeStack.pop_back();
-						break;
-					case EOpcodeType::Call:
-					{
-						const FCall Call = Opcodes[i].Call;
-						const int Arg = Opcodes[i].Value;
+				default:
+					FAIL("Unhandled opcode kind!");
+					break;
+				case EOpcodeType::StartTransaction:
+					AutoRTFM::StartTransaction();
+					RecoveryOpcodeStack.push_back(Opcodes[i].Value);
+					break;
+				case EOpcodeType::CommitTransaction:
+					REQUIRE(AutoRTFM::ETransactionResult::Committed == AutoRTFM::CommitTransaction());
+					RecoveryOpcodeStack.pop_back();
+					break;
+				case EOpcodeType::Call:
+				{
+					const FCall Call = Opcodes[i].Call;
+					const int Arg = Opcodes[i].Value;
 
-						const AutoRTFM::EContextStatus Status = AutoRTFM::Close([&Value, Arg, Call]
-							{
-								Call(&Value, Arg);
-							});
+					const AutoRTFM::EContextStatus Status = AutoRTFM::Close([&Value, Arg, Call]
+						{
+							Call(&Value, Arg);
+						});
 
-						if (AutoRTFM::EContextStatus::OnTrack != Status)
+					if (AutoRTFM::EContextStatus::OnTrack != Status)
+					{
+						if (RecoveryOpcodeStack.empty())
 						{
-							if (RecoveryOpcodeStack.empty())
-							{
-								ShouldThrowNest = true;
-								break;
-							}
-
-							// We've handled the bad transaction status here, so clear it out!
-							AutoRTFM::ClearTransactionStatus();
-
-							// -1 just cause the loop will i++!
-							i = RecoveryOpcodeStack.back() - 1;
-							RecoveryOpcodeStack.pop_back();
+							// We're wanting to throw out to the parent Transact call!
+							return;
 						}
-						break;
-					}
-					case EOpcodeType::WriteToValue:
-					{
-						const int Arg = Value - Opcodes[i].Value;
-						AutoRTFM::WriteMemory(&Value, &Arg);
-						break;
-					}
-					}
 
-					if (ShouldThrowNest)
-					{
-						break;
+						// We've handled the bad transaction status here, so clear it out!
+						AutoRTFM::ClearTransactionStatus();
+
+						// -1 just cause the loop will i++!
+						i = RecoveryOpcodeStack.back() - 1;
+						RecoveryOpcodeStack.pop_back();
 					}
+					break;
+				}
+				case EOpcodeType::WriteToValue:
+				{
+					const int Arg = Value - Opcodes[i].Value;
+					AutoRTFM::WriteMemory(&Value, &Arg);
+					break;
+				}
 				}
-			}
-
-			if (ShouldThrowNest)
-			{
-				AutoRTFM::CurrentNestThrow();
 			}
 		});
 	});
diff --git a/Engine/Source/Programs/AutoRTFMTests/Private/OpenAPI.cpp b/Engine/Source/Programs/AutoRTFMTests/Private/OpenAPI.cpp
--- a/Engine/Source/Programs/AutoRTFMTests/Private/OpenAPI.cpp
+++ b/Engine/Source/Programs/AutoRTFMTests/Private/OpenAPI.cpp
@@ -30,8 +30,6 @@ TEST_CASE("OpenAPI.StartAbortAndStartAgain")
 			AutoRTFM::RecordOpenWrite(&valueC);
 			valueC = 30;
 			AutoRTFM::AbortTransaction();
-
-			AutoRTFM::CurrentNestThrow();
 		});
 	});
 
@@ -87,7 +85,6 @@ TEST_CASE("OpenAPI.AbortTransactionScopedFromOpen")
 	{
 		AutoRTFM::Open([&]() {
 			AutoRTFM::AbortTransaction();
-			AutoRTFM::CurrentNestThrow();
 		});
 		FAIL("AutoRTFM::Open failed to throw after an abort");
 	});
@@ -118,7 +115,6 @@ TEST_CASE("OpenAPI.AbortTransactionDoubleScopedFromOpen")
 			value = 42424242;
 			AutoRTFM::Open([&]() {
 				AutoRTFM::AbortTransaction();
-				AutoRTFM::CurrentNestThrow();
 			});
 			FAIL("AutoRTFM::Open failed to throw after an abort");
 			value = 24242424;
@@ -376,7 +372,6 @@ TEST_CASE("OpenAPI.Commit_TransactOpenCloseAbort")
 			}));
 
 			AutoRTFM::AbortTransaction(); // undoes value = 42 in the open
-			AutoRTFM::CurrentNestThrow();
 			REQUIRE(valueLocal == 1.0);
 		});
 
@@ -504,7 +499,6 @@ TEST_CASE("OpenAPI.StackWriteAbortInTheOpen1")
 			value = 10;
 			AutoRTFM::AbortTransaction();
 			REQUIRE(value == 0);
-			AutoRTFM::CurrentNestThrow();
 		});
 	});
 }
@@ -556,7 +550,6 @@ TEST_CASE("OpenAPI.StackWriteAbortInTheOpen2")
 			AutoRTFM::WriteMemory(&value, 10); // Illegal to write to value because it's in the inner-most closed-nest
 
 			AutoRTFM::AbortTransaction();
-			AutoRTFM::CurrentNestThrow();
 		});
 
 		// Never gets here
@@ -609,7 +602,6 @@ TEST_CASE("OpenAPI.WriteTrivialStructure")
 			REQUIRE(data.E[2] == 345);
 			REQUIRE(data.E[3] == 456);
 			REQUIRE(data.E[4] == 567);
-			AutoRTFM::CurrentNestThrow();
 		});
 	});
 }
@@ -667,7 +659,6 @@ TEST_CASE("OpenAPI.WriteTrivialStructure2")
 				REQUIRE(data.E[2] == 345);
 				REQUIRE(data.E[3] == 456);
 				REQUIRE(data.E[4] == 567);
-				AutoRTFM::CurrentNestThrow();
 			});
 	});
 }
@@ -795,9 +786,7 @@ TEST_CASE("OpenAPI.TransOpenStartCloseAbortAbort")
 
 			REQUIRE(CloseStatus == AutoRTFM::EContextStatus::AbortedByRequest);
 
-			REQUIRE(AutoRTFM::IsAborting());
 			AutoRTFM::ClearTransactionStatus();
-			REQUIRE(!AutoRTFM::IsAborting());
 
 			REQUIRE(bGetsToA == false);
 
@@ -806,7 +795,6 @@ TEST_CASE("OpenAPI.TransOpenStartCloseAbortAbort")
 			AutoRTFM::AbortTransaction();
 			bGetsToC = true;
 			REQUIRE(value == 42);
-			AutoRTFM::CurrentNestThrow();
 		});
 
 		bGetsToD = true;
diff --git a/Engine/Source/Runtime/Core/Private/AutoRTFM/API.cpp b/Engine/Source/Runtime/Core/Private/AutoRTFM/API.cpp
--- a/Engine/Source/Runtime/Core/Private/AutoRTFM/API.cpp
+++ b/Engine/Source/Runtime/Core/Private/AutoRTFM/API.cpp
@@ -358,19 +358,6 @@ void STM_autortfm_clear_transaction_status(FContext* Context)
 	abort();
 }
 
-bool STM_autortfm_is_aborting(FContext* Context)
-{
-	return Context->IsAborting();
-}
-UE_AUTORTFM_REGISTER_OPEN_FUNCTION(autortfm_is_aborting);
-
-bool STM_autortfm_current_nest_throw(FContext* Context)
-{
-	Context->Throw();
-	return true;
-}
-UE_AUTORTFM_REGISTER_OPEN_FUNCTION(autortfm_current_nest_throw);
-
 void STM_autortfm_abort_if_transactional(FContext* Context)
 {
     if (FDebug::bVerbose)
@@ -394,8 +381,12 @@ UE_AUTORTFM_REGISTER_OPEN_FUNCTION(autortfm_abort_if_closed);
 
 void STM_autortfm_open(void (*Work)(void* Arg), void* Arg, FContext* Context)
 {
-	// WARNING! DO NOT EDIT! Changes to this function will be elided due to special compiler optimizations
 	Work(Arg);
+
+	if (Context->IsAborting())
+	{
+		Context->Throw();
+	}
 }
 UE_AUTORTFM_REGISTER_OPEN_FUNCTION(autortfm_open);
 
diff --git a/Engine/Source/Runtime/Core/Public/AutoRTFM/AutoRTFM.h b/Engine/Source/Runtime/Core/Public/AutoRTFM/AutoRTFM.h
--- a/Engine/Source/Runtime/Core/Public/AutoRTFM/AutoRTFM.h
+++ b/Engine/Source/Runtime/Core/Public/AutoRTFM/AutoRTFM.h
@@ -17,9 +17,10 @@
 #include <stdbool.h>
 #include <stdlib.h>
 #include <memory.h>
+#include <type_traits>
 
 #if UE_AUTORTFM
-#include "Templates/Function.h"
+template <typename FuncType> class TFunction;
 #endif
 
 #ifdef __cplusplus
@@ -154,24 +155,6 @@ inline void autortfm_clear_transaction_status()
 }
 #endif
 
-#if UE_AUTORTFM
-bool autortfm_is_aborting();
-#else
-inline bool autortfm_is_aborting()
-{
-	return false;
-}
-#endif
-
-#if UE_AUTORTFM
-bool autortfm_current_nest_throw();
-#else
-inline bool autortfm_current_nest_throw()
-{
-	return false;
-}
-#endif
-
 // Abort if running in a transaction. If called from closed code, this does
 // an abort_by_language, which will either abort the program or the transaction,
 // depending on configuration. If called from open code that is inside a
@@ -401,16 +384,6 @@ inline void ClearTransactionStatus()
 	autortfm_clear_transaction_status();
 }
 
-inline bool IsAborting()
-{
-	return autortfm_is_aborting();
-}
-
-inline bool CurrentNestThrow()
-{
-	return autortfm_current_nest_throw();
-}
-
 // RecordOpenWrite records the memory span into the current transaction as written.
 //  If this memory is previously unknown to the transaction, the original value is saved.
 inline void RecordOpenWrite(void* Ptr, size_t Size)
