Remove DeferUntilCommit / DeferUntilAbort from the AutoRTFM API (use OpenCommit / OpenAbort instead).

diff --git a/Engine/Source/Programs/AutoRTFMTests/Private/API.cpp b/Engine/Source/Programs/AutoRTFMTests/Private/API.cpp
--- a/Engine/Source/Programs/AutoRTFMTests/Private/API.cpp
+++ b/Engine/Source/Programs/AutoRTFMTests/Private/API.cpp
@@ -276,152 +276,6 @@ TEST_CASE("API.autortfm_register_open_function")
     REQUIRE(0 == I);
 }
 
-TEST_CASE("API.autortfm_defer_until_commit")
-{
-    bool OuterTransaction = false;
-    bool InnerTransaction = false;
-    bool InnerTransactionWithAbort = false;
-    bool InnerOpenNest = false;
-    AutoRTFM::ETransactionResult NestResult;
-
-    REQUIRE(AutoRTFM::ETransactionResult::Committed == AutoRTFM::Transact([&]
-    {
-        autortfm_defer_until_commit([](void* const Arg)
-        {
-            *static_cast<bool* const>(Arg) = true;
-        }, &OuterTransaction);
-
-        // This should only be modified on the commit!
-        if (OuterTransaction)
-        {
-            AutoRTFM::AbortTransaction();
-        }
-
-        AutoRTFM::Commit([&]
-        {
-            autortfm_defer_until_commit([](void* const Arg)
-            {
-                *static_cast<bool* const>(Arg) = true;
-            }, &InnerTransaction);
-        });
-
-        // This should only be modified on the commit!
-        if (InnerTransaction)
-        {
-            AutoRTFM::AbortTransaction();
-        }
-
-        NestResult = AutoRTFM::Transact([&]
-        {
-            autortfm_defer_until_commit([](void* const Arg)
-            {
-                *static_cast<bool* const>(Arg) = true;
-            }, &InnerTransactionWithAbort);
-
-            AutoRTFM::AbortTransaction();
-        });
-
-        // This should never be modified because its transaction aborted!
-        if (InnerTransactionWithAbort)
-        {
-            AutoRTFM::AbortTransaction();
-        }
-
-        AutoRTFM::Open([&]
-        {
-            autortfm_defer_until_commit([](void* const Arg)
-            {
-                *static_cast<bool* const>(Arg) = true;
-            }, &InnerOpenNest);
-        });
-
-        // This should only be modified on the commit!
-        if (InnerOpenNest)
-        {
-            AutoRTFM::AbortTransaction();
-        }
-    }));
-
-    REQUIRE(true == OuterTransaction);
-    REQUIRE(true == InnerTransaction);
-    REQUIRE(false == InnerTransactionWithAbort);
-    REQUIRE(true == InnerOpenNest);
-    REQUIRE(AutoRTFM::ETransactionResult::AbortedByRequest == NestResult);
-}
-
-TEST_CASE("API.autortfm_defer_until_abort")
-{
-    bool OuterTransaction = false;
-    bool InnerTransaction = false;
-    bool InnerTransactionWithAbort = false;
-    bool InnerOpenNest = false;
-    AutoRTFM::ETransactionResult NestResult;
-
-    REQUIRE(AutoRTFM::ETransactionResult::Committed == AutoRTFM::Transact([&]
-    {
-        autortfm_defer_until_abort([](void* const Arg)
-        {
-            *static_cast<bool* const>(Arg) = true;
-        }, &OuterTransaction);
-
-        // This should only be modified on the commit!
-        if (OuterTransaction)
-        {
-            AutoRTFM::AbortTransaction();
-        }
-
-        AutoRTFM::Commit([&]
-        {
-            autortfm_defer_until_abort([](void* const Arg)
-            {
-                *static_cast<bool* const>(Arg) = true;
-            }, &InnerTransaction);
-        });
-
-        // This should only be modified on the commit!
-        if (InnerTransaction)
-        {
-            AutoRTFM::AbortTransaction();
-        }
-
-        NestResult = AutoRTFM::Transact([&]
-        {
-            autortfm_defer_until_abort([](void* const Arg)
-            {
-                *static_cast<bool* const>(Arg) = true;
-            }, &InnerTransactionWithAbort);
-
-            AutoRTFM::AbortTransaction();
-        });
-
-        // This should never be modified because its transaction aborted!
-        if (InnerTransactionWithAbort)
-        {
-            AutoRTFM::AbortTransaction();
-        }
-
-        AutoRTFM::Open([&]
-        {
-            autortfm_defer_until_abort([](void* const Arg)
-            {
-                *static_cast<bool* const>(Arg) = true;
-            }, &InnerOpenNest);
-        });
-
-        // This should only be modified on the commit!
-        if (InnerOpenNest)
-        {
-            AutoRTFM::AbortTransaction();
-        }
-    }));
-
-    REQUIRE(false == OuterTransaction);
-    REQUIRE(false == InnerTransaction);
-    REQUIRE(true == InnerTransactionWithAbort);
-    REQUIRE(false == InnerOpenNest);
-    REQUIRE(AutoRTFM::ETransactionResult::AbortedByRequest == NestResult);
-}
-
 TEST_CASE("API.autortfm_open_commit")
 {
     bool OuterTransaction = false;
@@ -871,152 +725,6 @@ TEST_CASE("API.RegisterOpenFunction")
     REQUIRE(0 == I);
 }
 
-TEST_CASE("API.DeferUntilCommit")
-{
-    bool OuterTransaction = false;
-    bool InnerTransaction = false;
-    bool InnerTransactionWithAbort = false;
-    bool InnerOpenNest = false;
-    AutoRTFM::ETransactionResult NestResult;
-
-    REQUIRE(AutoRTFM::ETransactionResult::Committed == AutoRTFM::Transact([&]
-    {
-        AutoRTFM::DeferUntilCommit([&]
-        {
-            OuterTransaction = true;
-        });
-
-        // This should only be modified on the commit!
-        if (OuterTransaction)
-        {
-            AutoRTFM::AbortTransaction();
-        }
-
-        AutoRTFM::Commit([&]
-        {
-            AutoRTFM::DeferUntilCommit([&]
-            {
-                InnerTransaction = true;
-            });
-        });
-
-        // This should only be modified on the commit!
-        if (InnerTransaction)
-        {
-            AutoRTFM::AbortTransaction();
-        }
-
-        NestResult = AutoRTFM::Transact([&]
-        {
-            AutoRTFM::DeferUntilCommit([&]
-            {
-                InnerTransactionWithAbort = true;
-            });
-
-            AutoRTFM::AbortTransaction();
-        });
-
-        // This should never be modified because its transaction aborted!
-        if (InnerTransactionWithAbort)
-        {
-            AutoRTFM::AbortTransaction();
-        }
-
-        AutoRTFM::Open([&]
-        {
-            AutoRTFM::DeferUntilCommit([&]
-            {
-                InnerOpenNest = true;
-            });
-        });
-
-        // This should only be modified on the commit!
-        if (InnerOpenNest)
-        {
-            AutoRTFM::AbortTransaction();
-        }
-    }));
-
-    REQUIRE(true == OuterTransaction);
-    REQUIRE(true == InnerTransaction);
-    REQUIRE(false == InnerTransactionWithAbort);
-    REQUIRE(true == InnerOpenNest);
-    REQUIRE(AutoRTFM::ETransactionResult::AbortedByRequest == NestResult);
-}
-
-TEST_CASE("API.DeferUntilAbort")
-{
-    bool OuterTransaction = false;
-    bool InnerTransaction = false;
-    bool InnerTransactionWithAbort = false;
-    bool InnerOpenNest = false;
-    AutoRTFM::ETransactionResult NestResult;
-
-    REQUIRE(AutoRTFM::ETransactionResult::Committed == AutoRTFM::Transact([&]
-    {
-        AutoRTFM::DeferUntilAbort([&]
-        {
-            OuterTransaction = true;
-        });
-
-        // This should only be modified on the commit!
-        if (OuterTransaction)
-        {
-            AutoRTFM::AbortTransaction();
-        }
-
-        AutoRTFM::Commit([&]
-        {
-            AutoRTFM::DeferUntilAbort([&]
-            {
-                InnerTransaction = true;
-            });
-        });
-
-        // This should only be modified on the commit!
-        if (InnerTransaction)
-        {
-			AutoRTFM::AbortTransaction();
-        }
-
-        NestResult = AutoRTFM::Transact([&]
-        {
-            AutoRTFM::DeferUntilAbort([&]
-            {
-                InnerTransactionWithAbort = true;
-            });
-
-            AutoRTFM::AbortTransaction();
-        });
-
-        // This should never be modified because its transaction aborted!
-        if (InnerTransactionWithAbort)
-        {
-            AutoRTFM::AbortTransaction();
-        }
-
-        AutoRTFM::Open([&]
-        {
-            AutoRTFM::DeferUntilAbort([&]
-            {
-                InnerOpenNest = true;
-            });
-        });
-
-        // This should only be modified on the commit!
-        if (InnerOpenNest)
-        {
-			AutoRTFM::AbortTransaction();
-        }
-    }));
-
-    REQUIRE(false == OuterTransaction);
-    REQUIRE(false == InnerTransaction);
-    REQUIRE(true == InnerTransactionWithAbort);
-    REQUIRE(false == InnerOpenNest);
-    REQUIRE(AutoRTFM::ETransactionResult::AbortedByRequest == NestResult);
-}
-
 TEST_CASE("API.OpenCommit")
 {
     bool OuterTransaction = false;
diff --git a/Engine/Source/Programs/AutoRTFMTests/Private/DeferUntilCommit.cpp b/Engine/Source/Programs/AutoRTFMTests/Private/DeferUntilCommit.cpp
deleted file mode 100644
diff --git a/Engine/Source/Runtime/Core/Private/AutoRTFM/API.cpp b/Engine/Source/Runtime/Core/Private/AutoRTFM/API.cpp
--- a/Engine/Source/Runtime/Core/Private/AutoRTFM/API.cpp
+++ b/Engine/Source/Runtime/Core/Private/AutoRTFM/API.cpp
@@ -125,26 +125,6 @@ extern "C" void autortfm_register_open_function(void* OriginalFunction, void* Ne
     FunctionMapAdd(OriginalFunction, NewFunction);
 }
 
-void DeferUntilCommit(TFunction<void()>&& Work)
-{
-    if (FContext::IsTransactional())
-    {
-        FContext::Get()->GetCurrentTransaction()->DeferUntilCommit(MoveTemp(Work));
-    }
-    else
-    {
-        Work();
-    }
-}
-
-void DeferUntilAbort(TFunction<void()>&& Work)
-{
-    if (FContext::IsTransactional())
-    {
-        FContext::Get()->GetCurrentTransaction()->DeferUntilAbort(MoveTemp(Work));
-    }
-}
-
 void OpenCommit(TFunction<void()>&& Work)
 {
     Work();
@@ -154,16 +134,6 @@ void OpenAbort(TFunction<void()>&& Work)
 {
 }
 
-extern "C" void autortfm_defer_until_commit(void (*Work)(void* Arg), void* Arg)
-{
-    DeferUntilCommit([Work, Arg] () { Work(Arg); });
-}
-
-extern "C" void autortfm_defer_until_abort(void (*Work)(void* arg), void* Arg)
-{
-    DeferUntilAbort([Work, Arg] () { Work(Arg); });
-}
-
 extern "C" void autortfm_open_commit(void (*Work)(void* Arg), void* Arg)
 {
     Work(Arg);
@@ -333,20 +303,6 @@ extern "C" void RTFM_autortfm_record_open_write(void*, size_t, FContext*)
 }
 UE_AUTORTFM_REGISTER_OPEN_FUNCTION(autortfm_record_open_write);
 
-void RTFM_DeferUntilCommit(TFunction<void()>&& Work, FContext* Context)
-{
-    ASSERT(Context->GetStatus() == EContextStatus::OnTrack);
-    Context->GetCurrentTransaction()->DeferUntilCommit(MoveTemp(Work));
-}
-UE_AUTORTFM_REGISTER_OPEN_FUNCTION(DeferUntilCommit);
-
-void RTFM_DeferUntilAbort(TFunction<void()>&& Work, FContext* Context)
-{
-    ASSERT(Context->GetStatus() == EContextStatus::OnTrack);
-    Context->GetCurrentTransaction()->DeferUntilAbort(MoveTemp(Work));
-}
-UE_AUTORTFM_REGISTER_OPEN_FUNCTION(DeferUntilAbort);
-
 void RTFM_OpenCommit(TFunction<void()>&& Work, FContext* Context)
 {
     ASSERT(Context->GetStatus() == EContextStatus::OnTrack);
@@ -361,18 +317,6 @@ void RTFM_OpenAbort(TFunction<void()>&& Work, FContext* Context)
 }
 UE_AUTORTFM_REGISTER_OPEN_FUNCTION(OpenAbort);
 
-extern "C" void RTFM_autortfm_defer_until_commit(void (*Work)(void* Arg), void* Arg, FContext* Context)
-{
-    RTFM_DeferUntilCommit([Work, Arg] { Work(Arg); }, Context);
-}
-UE_AUTORTFM_REGISTER_OPEN_FUNCTION(autortfm_defer_until_commit);
-
-extern "C" void RTFM_autortfm_defer_until_abort(void (*Work)(void* arg), void* Arg, FContext* Context)
-{
-    RTFM_DeferUntilAbort([Work, Arg] { Work(Arg); }, Context);
-}
-UE_AUTORTFM_REGISTER_OPEN_FUNCTION(autortfm_defer_until_abort);
-
 extern "C" void RTFM_autortfm_open_commit(void (*Work)(void* Arg), void* Arg, FContext* Context)
 {
     RTFM_OpenCommit([Work, Arg] { Work(Arg); }, Context);
diff --git a/Engine/Source/Runtime/Core/Public/AutoRTFM/AutoRTFM.h b/Engine/Source/Runtime/Core/Public/AutoRTFM/AutoRTFM.h
--- a/Engine/Source/Runtime/Core/Public/AutoRTFM/AutoRTFM.h
+++ b/Engine/Source/Runtime/Core/Public/AutoRTFM/AutoRTFM.h
@@ -260,33 +260,6 @@ UE_AUTORTFM_FORCEINLINE void autortfm_register_open_function(void* original_func
 }
 #endif
 
-// Have some work happen when this transaction commits. For nested transactions,
-// this just adds the work to the work deferred until the outer nest's commit.
-// If this is called outside a transaction then the work happens immediately. An
-// open nest within a transaction will be deferred until commit too.
-#if UE_AUTORTFM
-void autortfm_defer_until_commit(void (*work)(void* arg), void* arg);
-#else
-UE_AUTORTFM_FORCEINLINE void autortfm_defer_until_commit(void (*work)(void* arg), void* arg)
-{
-    work(arg);
-}
-#endif
-
-// Have some work happen when this transaction aborts. For nested transactions,
-// this work will be processed only when the outer nest actually commits or
-// aborts. If this is called outside a transaction then the work is ignored. An
-// open nest within a transaction will be deferred until abort too.
-#if UE_AUTORTFM
-void autortfm_defer_until_abort(void (*work)(void* arg), void* arg);
-#else
-UE_AUTORTFM_FORCEINLINE void autortfm_defer_until_abort(void (*work)(void* arg), void* arg)
-{
-	UE_AUTORTFM_UNUSED(work);
-	UE_AUTORTFM_UNUSED(arg);
-}
-#endif
-
 // Have some work happen when this transaction commits. For nested transactions,
 // this just adds the work to the work deferred until the outer nest's commit.
 // If this is called outside a transaction or from an open nest then the work
@@ -515,16 +488,10 @@ UE_AUTORTFM_FORCEINLINE void RegisterOpenFunction(void* OriginalFunction, void*
 }
 
 #if UE_AUTORTFM
-void DeferUntilCommit(TFunction<void()>&& Work);
-void DeferUntilAbort(TFunction<void()>&& Work);
 void OpenCommit(TFunction<void()>&& Work);
 void OpenAbort(TFunction<void()>&& Work);
 #else
 template<typename TFunctor>
-UE_AUTORTFM_FORCEINLINE void DeferUntilCommit(const TFunctor& Work) { Work(); }
-template<typename TFunctor>
-UE_AUTORTFM_FORCEINLINE void DeferUntilAbort(const TFunctor& Work) { }
-template<typename TFunctor>
 UE_AUTORTFM_FORCEINLINE void OpenCommit(const TFunctor& Work) { Work(); }
 template<typename TFunctor>
 UE_AUTORTFM_FORCEINLINE void OpenAbort(const TFunctor& Work) { }
