Fix that in AutoRTFM we need open aborts inside a nested transaction to be appended to the parent in reverse order.

diff --git a/Engine/Source/Programs/AutoRTFMTests/Private/Abort.cpp b/Engine/Source/Programs/AutoRTFMTests/Private/Abort.cpp
--- a/Engine/Source/Programs/AutoRTFMTests/Private/Abort.cpp
+++ b/Engine/Source/Programs/AutoRTFMTests/Private/Abort.cpp
@@ -49,3 +49,32 @@ TEST_CASE("Abort")
     REQUIRE(m[6][1] == 8);
     REQUIRE(m[6][2] == 9);
 }
+
+TEST_CASE("Abort.NestedAbortOrder")
+{
+	AutoRTFM::ETransactionResult InnerResult;
+	unsigned Orderer = 0;
+
+	AutoRTFM::Commit([&]
+		{
+			InnerResult = AutoRTFM::Transact([&]
+				{
+					AutoRTFM::OpenAbort([&]
+						{
+							REQUIRE(1 == Orderer);
+							Orderer += 1;
+						});
+
+					AutoRTFM::OpenAbort([&]
+						{
+							REQUIRE(0 == Orderer);
+							Orderer += 1;
+						});
+
+					AutoRTFM::AbortTransaction();
+				});
+		});
+
+	REQUIRE(AutoRTFM::ETransactionResult::AbortedByRequest == InnerResult);
+	REQUIRE(2 == Orderer);
+}
diff --git a/Engine/Source/Runtime/Core/Private/AutoRTFM/Transaction.cpp b/Engine/Source/Runtime/Core/Private/AutoRTFM/Transaction.cpp
--- a/Engine/Source/Runtime/Core/Private/AutoRTFM/Transaction.cpp
+++ b/Engine/Source/Runtime/Core/Private/AutoRTFM/Transaction.cpp
@@ -135,7 +135,9 @@ void FTransaction::AbortNested()
 
     Undo();
 
-    Parent->CommitTasks.AddAll(AbortTasks);
+	// We need to add the abort tasks in reverse order to the parent, as they need to be run in the reverse order.
+	AbortTasks.ForEachBackward([&](const TFunction<void()>& Task) -> bool { Parent->CommitTasks.Add(Task); return true; });
+
     Parent->AbortTasks.AddAll(MoveTemp(AbortTasks));
 }
 
diff --git a/Engine/Source/Runtime/Core/Private/AutoRTFM/TransactionInlines.h b/Engine/Source/Runtime/Core/Private/AutoRTFM/TransactionInlines.h
--- a/Engine/Source/Runtime/Core/Private/AutoRTFM/TransactionInlines.h
+++ b/Engine/Source/Runtime/Core/Private/AutoRTFM/TransactionInlines.h
@@ -57,7 +57,7 @@ inline void FTransaction::RecordWrite(void* LogicalAddress, size_t Size)
 
 inline void FTransaction::DidAllocate(void* LogicalAddress, size_t Size)
 {
-    if (Size <= FWriteLogBumpAllocator::MaxSize)
+    if ((0 < Size) && (Size <= FWriteLogBumpAllocator::MaxSize))
     {
         FMemoryLocation Key(LogicalAddress);
         Key.SetTopTag(static_cast<uint16_t>(Size));
