Add some UObject tests to AutoRTFMTests.

diff --git a/Engine/Source/Programs/AutoRTFMTests/AutoRTFMTests.Build.cs b/Engine/Source/Programs/AutoRTFMTests/AutoRTFMTests.Build.cs
--- a/Engine/Source/Programs/AutoRTFMTests/AutoRTFMTests.Build.cs
+++ b/Engine/Source/Programs/AutoRTFMTests/AutoRTFMTests.Build.cs
@@ -12,6 +12,7 @@ public class AutoRTFMTests : ModuleRules
 			new string[] {
 				"Catch2Extras",
 				"Core",
+				"CoreUObject",
 				"ApplicationCore",
 				"Projects",
 			}
diff --git a/Engine/Source/Programs/AutoRTFMTests/AutoRTFMTests.Target.cs b/Engine/Source/Programs/AutoRTFMTests/AutoRTFMTests.Target.cs
--- a/Engine/Source/Programs/AutoRTFMTests/AutoRTFMTests.Target.cs
+++ b/Engine/Source/Programs/AutoRTFMTests/AutoRTFMTests.Target.cs
@@ -15,9 +15,8 @@ public class AutoRTFMTestsTarget : TargetRules
 		Name = "AutoRTFMTests";
 		LaunchModuleName = "AutoRTFMTests";
 
-		// Compile out references from Core to the rest of the engine
 		bCompileAgainstEngine = false;
-		bCompileAgainstCoreUObject = false;
+		bCompileAgainstCoreUObject = true;
 
 		// Logs are still useful to print the results
 		bUseLoggingInShipping = true;
diff --git a/Engine/Source/Programs/AutoRTFMTests/Classes/MyAutoRTFMTestObject.h b/Engine/Source/Programs/AutoRTFMTests/Classes/MyAutoRTFMTestObject.h
new file mode 100644
--- /dev/null
+++ b/Engine/Source/Programs/AutoRTFMTests/Classes/MyAutoRTFMTestObject.h
@@ -0,0 +1,16 @@
+// Copyright Epic Games, Inc. All Rights Reserved.
+#pragma once
+
+#include "MyAutoRTFMTestObject.generated.h"
+
+UCLASS()
+class UMyAutoRTFMTestObject : public UObject
+{
+	GENERATED_BODY()
+
+public:
+	UMyAutoRTFMTestObject() : Value(42)
+	{}
+
+	int Value;
+};
diff --git a/Engine/Source/Programs/AutoRTFMTests/Private/AutoRTFMTests.cpp b/Engine/Source/Programs/AutoRTFMTests/Private/AutoRTFMTests.cpp
--- a/Engine/Source/Programs/AutoRTFMTests/Private/AutoRTFMTests.cpp
+++ b/Engine/Source/Programs/AutoRTFMTests/Private/AutoRTFMTests.cpp
@@ -11,10 +11,22 @@ class SetupListener : public Catch::EventListenerBase {
 public:
 	using Catch::EventListenerBase::EventListenerBase;
 
-	void testRunStarting(Catch::TestRunInfo const&) override
+	void testRunStarting(const Catch::TestRunInfo&) override
 	{
+		GEngineLoop.PreInit(0, nullptr);
+		FModuleManager::Get().StartProcessingNewlyLoadedObjects();
+
 		AutoRTFM::SetAutoRTFMRuntime(true);
 	}
+
+	void testRunEnded(const Catch::TestRunStats&) override
+	{
+		FEngineLoop::AppPreExit();
+		FModuleManager::Get().UnloadModulesAtShutdown();
+		FEngineLoop::AppExit();
+
+		FPlatformMisc::RequestExit(false);
+	}
 };
 
 CATCH_REGISTER_LISTENER(SetupListener)
diff --git a/Engine/Source/Programs/AutoRTFMTests/Private/UObjectTests.cpp b/Engine/Source/Programs/AutoRTFMTests/Private/UObjectTests.cpp
new file mode 100644
--- /dev/null
+++ b/Engine/Source/Programs/AutoRTFMTests/Private/UObjectTests.cpp
@@ -0,0 +1,102 @@
+// Copyright Epic Games, Inc. All Rights Reserved.
+
+#pragma autortfm
+
+#include "MyAutoRTFMTestObject.h"
+#include "Catch2Includes.h"
+#include <AutoRTFM/AutoRTFM.h>
+#include "UObject/GCObject.h"
+#include "UObject/ReachabilityAnalysis.h"
+
+TEST_CASE("UObject.Create")
+{
+	UMyAutoRTFMTestObject* Object = nullptr;
+
+	AutoRTFM::Commit([&]
+	{
+		Object = NewObject<UMyAutoRTFMTestObject>();
+	});
+
+	REQUIRE(nullptr != Object);
+	REQUIRE(42 == Object->Value);
+}
+
+TEST_CASE("UObject.Abort")
+{
+	UMyAutoRTFMTestObject* Object = nullptr;
+
+	REQUIRE(AutoRTFM::ETransactionResult::AbortedByRequest == AutoRTFM::Transact([&]
+	{
+		Object = NewObject<UMyAutoRTFMTestObject>();
+		AutoRTFM::AbortTransaction();
+	}));
+
+	REQUIRE(nullptr == Object);
+}
+
+// This is a copy of the helper function in TestGarbageCollector.cpp.
+int32 PerformGarbageCollectionWithIncrementalReachabilityAnalysis(TFunctionRef<bool(int32)> ReachabilityIterationCallback)
+{
+	int32 ReachabilityIterationIndex = 0;
+
+	CollectGarbage(GARBAGE_COLLECTION_KEEPFLAGS, false);
+
+	while (IsIncrementalReachabilityAnalysisPending())
+	{
+		if (ReachabilityIterationCallback(ReachabilityIterationIndex))
+		{
+			break;
+		}
+
+		// Re-check if incremental rachability is still pending because the callback above could've triggered GC which would complete all iterations
+		if (IsIncrementalReachabilityAnalysisPending())
+		{
+			PerformIncrementalReachabilityAnalysis(GetReachabilityAnalysisTimeLimit());
+			ReachabilityIterationIndex++;
+		}
+	}
+
+	if (IsIncrementalPurgePending())
+	{
+		IncrementalPurgeGarbage(false);
+	}
+	check(IsIncrementalPurgePending() == false);
+
+	return ReachabilityIterationIndex + 1;
+}
+
+TEST_CASE("UObject.MarkAsReachable")
+{
+	// We need incremental reachability to be on.
+	SetIncrementalReachabilityAnalysisEnabled(true);
+
+	// Cache the original time limit.
+	const float Original = GetReachabilityAnalysisTimeLimit();
+
+	// And we need a super small time limit s that reachability analysis will definitely have started.
+	SetReachabilityAnalysisTimeLimit(FLT_MIN);
+	
+	// We need to be sure we've done the static GC initialization before we start doing a garbage
+	// collection.
+	FGCObject::StaticInit();
+
+	UMyAutoRTFMTestObject* const Object = NewObject<UMyAutoRTFMTestObject>();
+
+	PerformGarbageCollectionWithIncrementalReachabilityAnalysis([Object](int32 index)
+	{
+		if (0 != index)
+		{
+			return true;
+		}
+
+		AutoRTFM::Commit([&]
+		{
+			Object->MarkAsReachable();
+		});
+
+		return false;
+	});
+
+	// Reset it back just incase another test required the original time limit.
+	SetReachabilityAnalysisTimeLimit(Original);
+}
diff --git a/Engine/Source/Runtime/CoreUObject/Private/UObject/GarbageCollection.cpp b/Engine/Source/Runtime/CoreUObject/Private/UObject/GarbageCollection.cpp
--- a/Engine/Source/Runtime/CoreUObject/Private/UObject/GarbageCollection.cpp
+++ b/Engine/Source/Runtime/CoreUObject/Private/UObject/GarbageCollection.cpp
@@ -5734,7 +5734,12 @@ FORCEINLINE static void MarkAsReachable(const UObject* Obj)
 
 void UObject::MarkAsReachable() const
 {
-	::MarkAsReachable<false>(this);
+	// It is safe to perform mark as reachable in the open - the worst case is that we'll mark an object reachable that
+	// should/would be destroyed, and so in the next GC iteration it will be destroyed instead of in this iteration.
+	UE_AUTORTFM_OPEN(
+	{
+		::MarkAsReachable<false>(this);
+	});
 }
 
 #if WITH_VERSE_VM || defined(__INTELLISENSE__)
