Make the AutoRTFM compiler support static local initializers.

diff --git a/Engine/Source/Programs/AutoRTFMTests/Private/Statics.cpp b/Engine/Source/Programs/AutoRTFMTests/Private/Statics.cpp
new file mode 100644
--- /dev/null
+++ b/Engine/Source/Programs/AutoRTFMTests/Private/Statics.cpp
@@ -0,0 +1,107 @@
+// Copyright Epic Games, Inc. All Rights Reserved.
+
+#pragma autortfm
+
+#include "Catch2Includes.h"
+#include <AutoRTFM/AutoRTFM.h>
+
+TEST_CASE("Statics.Int")
+{
+	auto GetAndIncrement = []
+		{
+			static int Thing = 42;
+			return Thing++;
+		};
+
+	REQUIRE(
+		AutoRTFM::ETransactionResult::Committed ==
+		AutoRTFM::Transact([&]()
+			{
+				GetAndIncrement();
+			}));
+
+	// The transactional effect of incrementing the static will have
+	// been committed, and since we are accessing the exact same
+	// static we should see its side effects.
+	REQUIRE(43 == GetAndIncrement());
+}
+
+TEST_CASE("Statics.IntAbort")
+{
+	auto GetAndIncrement = []
+		{
+			static int Thing = 42;
+			return Thing++;
+		};
+
+    REQUIRE(
+        AutoRTFM::ETransactionResult::AbortedByRequest ==
+        AutoRTFM::Transact([&] ()
+        {
+			if (42 == GetAndIncrement())
+			{
+				AutoRTFM::AbortTransaction();
+			}
+        }));
+    
+	// The transactional effect of incrementing the static will have
+	// been rolled back, but it should still be initialized correctly.
+	REQUIRE(42 == GetAndIncrement());
+}
+
+struct SomeStruct
+{
+	int Payload[42];
+	int Current = 0;
+};
+
+TEST_CASE("Statics.Struct")
+{
+	auto GetSlot = []
+		{
+			static SomeStruct S;
+			int* const Result = &S.Payload[S.Current];
+			*Result = ++S.Current;
+			return Result;
+		};
+
+	REQUIRE(
+		AutoRTFM::ETransactionResult::Committed ==
+		AutoRTFM::Transact([&]()
+			{
+				int* const Slot = GetSlot();
+				*Slot = 13;
+			}));
+
+	// The transactional effect of incrementing the static will have
+	// been committed, so we should see the side effects.
+	int* const Slot = GetSlot();
+	REQUIRE(2 == *Slot);
+
+	// The transaction would have written to the previous slot.
+	REQUIRE(13 == Slot[-1]);
+}
+
+TEST_CASE("Statics.StructAbort")
+{
+	auto GetSlot = []
+		{
+			static SomeStruct S;
+			int* const Result = &S.Payload[S.Current];
+			*Result = ++S.Current;
+			return Result;
+		};
+
+	REQUIRE(
+		AutoRTFM::ETransactionResult::AbortedByRequest ==
+		AutoRTFM::Transact([&]()
+			{
+				int* const Slot = GetSlot();
+				*Slot = 13;
+				AutoRTFM::AbortTransaction();
+			}));
+
+	// The transactional effect of incrementing the static will have
+	// been rolled back, but it should still be initialized correctly.
+	REQUIRE(1 == *GetSlot());
+}
diff --git a/Engine/Source/Runtime/Core/Public/AutoRTFM/AutoRTFM.h b/Engine/Source/Runtime/Core/Public/AutoRTFM/AutoRTFM.h
--- a/Engine/Source/Runtime/Core/Public/AutoRTFM/AutoRTFM.h
+++ b/Engine/Source/Runtime/Core/Public/AutoRTFM/AutoRTFM.h
@@ -23,7 +23,7 @@
 
 #if defined(UE_AUTORTFM_STANDALONE)
 #define UE_AUTORTFM_API
-#define UE_AUTORTFM_FORCEINLINE
+#define UE_AUTORTFM_FORCEINLINE inline
 #else
 #include <HAL/Platform.h>
 #define UE_AUTORTFM_API CORE_API
