Start to test every AutoRTFM use in UE core with a unit test to aid in preventing regressions.

diff --git a/Engine/Source/Programs/AutoRTFMTests/AutoRTFMTests.Build.cs b/Engine/Source/Programs/AutoRTFMTests/AutoRTFMTests.Build.cs
--- a/Engine/Source/Programs/AutoRTFMTests/AutoRTFMTests.Build.cs
+++ b/Engine/Source/Programs/AutoRTFMTests/AutoRTFMTests.Build.cs
@@ -26,5 +26,11 @@ public class AutoRTFMTests : ModuleRules
 		}
 
 		bAllowAutoRTFMInstrumentation = true;
+
+		PrivateIncludePaths.AddRange(
+			new string[]
+			{
+				"Runtime/Core/Private"
+			});
 	}
 }
diff --git a/Engine/Source/Programs/AutoRTFMTests/AutoRTFMTests.Target.cs b/Engine/Source/Programs/AutoRTFMTests/AutoRTFMTests.Target.cs
--- a/Engine/Source/Programs/AutoRTFMTests/AutoRTFMTests.Target.cs
+++ b/Engine/Source/Programs/AutoRTFMTests/AutoRTFMTests.Target.cs
@@ -32,5 +32,7 @@ public class AutoRTFMTestsTarget : TargetRules
 		{
 			 bUseAutoRTFMCompiler = true;
 		}
+
+		bFNameOutlineNumber = true;
 	}
 }
diff --git a/Engine/Source/Programs/AutoRTFMTests/Private/UECoreTests.cpp b/Engine/Source/Programs/AutoRTFMTests/Private/UECoreTests.cpp
new file mode 100644
--- /dev/null
+++ b/Engine/Source/Programs/AutoRTFMTests/Private/UECoreTests.cpp
@@ -0,0 +1,277 @@
+// Copyright Epic Games, Inc. All Rights Reserved.
+
+#pragma autortfm
+
+#include "Catch2Includes.h"
+#include "AutoRTFM/AutoRTFM.h"
+#include "Delegates/IDelegateInstance.h"
+#include "HAL/ThreadSingleton.h"
+#include "Internationalization/TextHistory.h"
+#include "Serialization/CustomVersion.h"
+#include "UObject/NameTypes.h"
+
+TEST_CASE("UECore.FDelegateHandle")
+{
+	FDelegateHandle Handle;
+
+	SECTION("With Abort")
+	{
+		AutoRTFM::ETransactionResult Result = AutoRTFM::Transact([&]()
+		{
+			Handle = FDelegateHandle(FDelegateHandle::GenerateNewHandle);
+			AutoRTFM::AbortTransaction();
+		});
+
+		REQUIRE(AutoRTFM::ETransactionResult::AbortedByRequest == Result);
+		REQUIRE(!Handle.IsValid());
+	}
+	
+	REQUIRE(!Handle.IsValid());
+
+	SECTION("With Commit")
+	{
+		AutoRTFM::ETransactionResult Result = AutoRTFM::Transact([&]()
+		{
+			Handle = FDelegateHandle(FDelegateHandle::GenerateNewHandle);
+		});
+
+		REQUIRE(AutoRTFM::ETransactionResult::Committed == Result);
+		REQUIRE(Handle.IsValid());
+	}
+}
+
+TEST_CASE("UECore.TThreadSingleton")
+{
+	struct MyStruct : TThreadSingleton<MyStruct>
+	{
+		int I;
+		float F;
+	};
+
+	SECTION("TryGet First Time")
+	{
+		REQUIRE(nullptr == TThreadSingleton<MyStruct>::TryGet());
+
+		// Set to something that isn't nullptr because TryGet will return that!
+		MyStruct* Singleton;
+		uintptr_t Data = 0x12345678abcdef00;
+		memcpy(&Singleton, &Data, sizeof(Singleton));
+
+		AutoRTFM::ETransactionResult Result = AutoRTFM::Transact([&]()
+		{
+			Singleton = TThreadSingleton<MyStruct>::TryGet();
+		});
+
+		REQUIRE(AutoRTFM::ETransactionResult::Committed == Result);
+		REQUIRE(nullptr == Singleton);
+	}
+
+	SECTION("Get")
+	{
+		AutoRTFM::ETransactionResult Result = AutoRTFM::Transact([&]()
+			{
+				TThreadSingleton<MyStruct>::Get().I = 42;
+				TThreadSingleton<MyStruct>::Get().F = 42.0f;
+				AutoRTFM::AbortTransaction();
+			});
+
+		REQUIRE(AutoRTFM::ETransactionResult::AbortedByRequest == Result);
+
+		// The singleton *will remain* initialized though, even though we got it in
+		// a transaction, because we have to do the singleton creation in the open.
+		REQUIRE(nullptr != TThreadSingleton<MyStruct>::TryGet());
+
+		// But any *changes* to the singleton data will be rolled back.
+		REQUIRE(0 == TThreadSingleton<MyStruct>::Get().I);
+		REQUIRE(0.0f == TThreadSingleton<MyStruct>::Get().F);
+
+		Result = AutoRTFM::Transact([&]()
+			{
+				TThreadSingleton<MyStruct>::Get().I = 42;
+				TThreadSingleton<MyStruct>::Get().F = 42.0f;
+			});
+
+		REQUIRE(AutoRTFM::ETransactionResult::Committed == Result);
+
+		REQUIRE(42 == TThreadSingleton<MyStruct>::Get().I);
+		REQUIRE(42.0f == TThreadSingleton<MyStruct>::Get().F);
+	}
+
+	SECTION("TryGet Second Time")
+	{
+		REQUIRE(nullptr != TThreadSingleton<MyStruct>::TryGet());
+
+		MyStruct* Singleton = nullptr;
+
+		AutoRTFM::ETransactionResult Result = AutoRTFM::Transact([&]()
+			{
+				Singleton = TThreadSingleton<MyStruct>::TryGet();
+			});
+
+		REQUIRE(AutoRTFM::ETransactionResult::Committed == Result);
+		REQUIRE(nullptr != Singleton);
+	}
+}
+
+TEST_CASE("UECore.FTextHistory")
+{
+	struct MyTextHistory final : FTextHistory_Base
+	{
+		// Need this to always return true so we hit the fun transactional bits!
+		bool CanUpdateDisplayString() override
+		{
+			return true;
+		}
+
+		MyTextHistory(const FTextId& InTextId, FString&& InSourceString) : FTextHistory_Base(InTextId, MoveTemp(InSourceString)) {}
+	};
+
+	FTextKey Namespace("NAMESPACE");
+	FTextKey Key("KEY");
+	FTextId TextId(Namespace, Key);
+	FString String("WOWWEE");
+
+	MyTextHistory History(TextId, MoveTemp(String));
+
+	SECTION("With Abort")
+	{
+		AutoRTFM::ETransactionResult Result = AutoRTFM::Transact([&]()
+			{
+				History.UpdateDisplayStringIfOutOfDate();
+				AutoRTFM::AbortTransaction();
+			});
+
+		REQUIRE(AutoRTFM::ETransactionResult::AbortedByRequest == Result);
+	}
+
+	SECTION("With Commit")
+	{
+		AutoRTFM::ETransactionResult Result = AutoRTFM::Transact([&]()
+			{
+				History.UpdateDisplayStringIfOutOfDate();
+			});
+
+		REQUIRE(AutoRTFM::ETransactionResult::Committed == Result);
+	}
+}
+
+TEST_CASE("UECore.FCustomVersionContainer")
+{
+	FCustomVersionContainer Container;
+	FGuid Guid(42, 42, 42, 42);
+
+	FCustomVersionRegistration Register(Guid, 0, TEXT("WOWWEE"));
+
+	REQUIRE(nullptr == Container.GetVersion(Guid));
+
+	SECTION("With Abort")
+	{
+		AutoRTFM::ETransactionResult Result = AutoRTFM::Transact([&]()
+			{
+				// The first time the version will be new.
+				Container.SetVersionUsingRegistry(Guid);
+
+				// The second time we should hit the cache the first one created.
+				Container.SetVersionUsingRegistry(Guid);
+				AutoRTFM::AbortTransaction();
+			});
+
+		REQUIRE(AutoRTFM::ETransactionResult::AbortedByRequest == Result);
+		REQUIRE(nullptr == Container.GetVersion(Guid));
+	}
+
+	SECTION("With Commit")
+	{
+		AutoRTFM::ETransactionResult Result = AutoRTFM::Transact([&]()
+			{
+				// The first time the version will be new.
+				Container.SetVersionUsingRegistry(Guid);
+
+				// The second time we should hit the cache the first one created.
+				Container.SetVersionUsingRegistry(Guid);
+			});
+
+		REQUIRE(AutoRTFM::ETransactionResult::Committed == Result);
+		REQUIRE(nullptr != Container.GetVersion(Guid));
+	}
+}
+
+TEST_CASE("UECore.FName")
+{
+	SECTION("EName Constructor")
+	{
+		FName Name;
+
+		SECTION("With Abort")
+		{
+			AutoRTFM::ETransactionResult Result = AutoRTFM::Transact([&]()
+				{
+					Name = FName(EName::Timer);
+					AutoRTFM::AbortTransaction();
+				});
+
+			REQUIRE(AutoRTFM::ETransactionResult::AbortedByRequest == Result);
+			REQUIRE(Name.IsNone());
+		}
+
+		SECTION("With Commit")
+		{
+			AutoRTFM::ETransactionResult Result = AutoRTFM::Transact([&]()
+				{
+					Name = FName(EName::Timer);
+				});
+
+			REQUIRE(AutoRTFM::ETransactionResult::Committed == Result);
+			REQUIRE(EName::Timer == *Name.ToEName());
+		}
+	}
+
+	SECTION("String Constructor")
+	{
+		FName Name;
+
+		SECTION("With Abort")
+		{
+			AutoRTFM::ETransactionResult Result = AutoRTFM::Transact([&]()
+				{
+					Name = FName(TEXT("WOWWEE"), 42);
+					AutoRTFM::AbortTransaction();
+				});
+
+			REQUIRE(AutoRTFM::ETransactionResult::AbortedByRequest == Result);
+			REQUIRE(Name.IsNone());
+		}
+
+		SECTION("Check FName was cached")
+		{
+			bool bWasCached = false;
+
+			for (const FNameEntry* const Entry : FName::DebugDump())
+			{
+				// Even though we aborted the transaction above, the actual backing data store of
+				// the FName system that deduplicates names will contain our name (the nature of
+				// the global shared caching infrastructure means we cannot just throw away the
+				// FName in the shared cache because it *could* have also been requested in the
+				// open and we'd be stomping on that legit use of it!).
+				if (0 != Entry->GetNameLength() && (TEXT("WOWWEE") == Entry->GetPlainNameString()))
+				{
+					bWasCached = true;
+				}
+			}
+
+			REQUIRE(bWasCached);
+		}
+
+		SECTION("With Commit")
+		{
+			AutoRTFM::ETransactionResult Result = AutoRTFM::Transact([&]()
+				{
+					Name = FName(TEXT("WOWWEE"), 42);
+				});
+
+			REQUIRE(AutoRTFM::ETransactionResult::Committed == Result);
+			REQUIRE(TEXT("WOWWEE") == Name.GetPlainNameString());
+			REQUIRE(42 == Name.GetNumber());
+		}
+	}
+}
