Make weak semantics of shared pointers work in AutoRTFM (and add tests).

diff --git a/Engine/Source/Programs/AutoRTFMTests/Private/UECoreTests.cpp b/Engine/Source/Programs/AutoRTFMTests/Private/UECoreTests.cpp
--- a/Engine/Source/Programs/AutoRTFMTests/Private/UECoreTests.cpp
+++ b/Engine/Source/Programs/AutoRTFMTests/Private/UECoreTests.cpp
@@ -275,3 +275,304 @@ TEST_CASE("UECore.FName")
 		}
 	}
 }
+
+TEST_CASE("UECore.STATIC_FUNCTION_FNAME")
+{
+	FName Name;
+
+	SECTION("With Abort")
+	{
+		AutoRTFM::ETransactionResult Result = AutoRTFM::Transact([&]()
+			{
+				Name = STATIC_FUNCTION_FNAME(TEXT("WOWWEE"));
+				AutoRTFM::AbortTransaction();
+			});
+
+		REQUIRE(AutoRTFM::ETransactionResult::AbortedByRequest == Result);
+		REQUIRE(Name.IsNone());
+	}
+
+	SECTION("With Commit")
+	{
+		AutoRTFM::ETransactionResult Result = AutoRTFM::Transact([&]()
+			{
+				Name = STATIC_FUNCTION_FNAME(TEXT("WOWWEE"));
+			});
+
+		REQUIRE(AutoRTFM::ETransactionResult::Committed == Result);
+	}
+}
+
+TEST_CASE("UECore.TIntrusiveReferenceController")
+{
+	SECTION("AddSharedReference")
+	{
+		SharedPointerInternals::TIntrusiveReferenceController<int, ESPMode::ThreadSafe> Controller(42);
+
+		SECTION("With Abort")
+		{
+			AutoRTFM::ETransactionResult Result = AutoRTFM::Transact([&]()
+				{
+					Controller.AddSharedReference();
+					AutoRTFM::AbortTransaction();
+				});
+
+			REQUIRE(AutoRTFM::ETransactionResult::AbortedByRequest == Result);
+			REQUIRE(1 == Controller.GetSharedReferenceCount());
+		}
+
+		SECTION("With Commit")
+		{
+			AutoRTFM::ETransactionResult Result = AutoRTFM::Transact([&]()
+				{
+					Controller.AddSharedReference();
+				});
+
+			REQUIRE(AutoRTFM::ETransactionResult::Committed == Result);
+			REQUIRE(2 == Controller.GetSharedReferenceCount());
+		}
+	}
+
+	SECTION("AddWeakReference")
+	{
+		SharedPointerInternals::TIntrusiveReferenceController<int, ESPMode::ThreadSafe> Controller(42);
+
+		SECTION("With Abort")
+		{
+			AutoRTFM::ETransactionResult Result = AutoRTFM::Transact([&]()
+				{
+					Controller.AddWeakReference();
+					AutoRTFM::AbortTransaction();
+				});
+
+			REQUIRE(AutoRTFM::ETransactionResult::AbortedByRequest == Result);
+			REQUIRE(1 == Controller.WeakReferenceCount);
+		}
+
+		SECTION("With Commit")
+		{
+			AutoRTFM::ETransactionResult Result = AutoRTFM::Transact([&]()
+				{
+					Controller.AddWeakReference();
+				});
+
+			REQUIRE(AutoRTFM::ETransactionResult::Committed == Result);
+			REQUIRE(2 == Controller.WeakReferenceCount);
+		}
+	}
+
+	SECTION("ConditionallyAddSharedReference")
+	{
+		SECTION("With Shared Reference Non Zero")
+		{
+			SharedPointerInternals::TIntrusiveReferenceController<int, ESPMode::ThreadSafe> Controller(42);
+
+			SECTION("With Abort")
+			{
+				AutoRTFM::ETransactionResult Result = AutoRTFM::Transact([&]()
+					{
+						Controller.ConditionallyAddSharedReference();
+						AutoRTFM::AbortTransaction();
+					});
+
+				REQUIRE(AutoRTFM::ETransactionResult::AbortedByRequest == Result);
+				REQUIRE(1 == Controller.GetSharedReferenceCount());
+			}
+
+			SECTION("With Commit")
+			{
+				AutoRTFM::ETransactionResult Result = AutoRTFM::Transact([&]()
+					{
+						Controller.ConditionallyAddSharedReference();
+					});
+
+				REQUIRE(AutoRTFM::ETransactionResult::Committed == Result);
+				REQUIRE(2 == Controller.GetSharedReferenceCount());
+			}
+		}
+
+		SECTION("With Shared Reference Zero")
+		{
+			SharedPointerInternals::TIntrusiveReferenceController<int, ESPMode::ThreadSafe> Controller(42);
+
+			// This test relies on us having a weak reference but no strong references to the object.
+			Controller.AddWeakReference();
+			Controller.ReleaseSharedReference();
+			REQUIRE(0 == Controller.GetSharedReferenceCount());
+
+			SECTION("With Abort")
+			{
+				AutoRTFM::ETransactionResult Result = AutoRTFM::Transact([&]()
+					{
+						Controller.ConditionallyAddSharedReference();
+						AutoRTFM::AbortTransaction();
+					});
+
+				REQUIRE(AutoRTFM::ETransactionResult::AbortedByRequest == Result);
+				REQUIRE(0 == Controller.GetSharedReferenceCount());
+			}
+
+			SECTION("With Commit")
+			{
+				AutoRTFM::ETransactionResult Result = AutoRTFM::Transact([&]()
+					{
+						Controller.ConditionallyAddSharedReference();
+					});
+
+				REQUIRE(AutoRTFM::ETransactionResult::Committed == Result);
+				REQUIRE(0 == Controller.GetSharedReferenceCount());
+			}
+		}
+	}
+
+	SECTION("GetSharedReferenceCount")
+	{
+		SharedPointerInternals::TIntrusiveReferenceController<int, ESPMode::ThreadSafe> Controller(42);
+
+		SECTION("With Abort")
+		{
+			int32 Count = 0;
+
+			AutoRTFM::ETransactionResult Result = AutoRTFM::Transact([&]()
+				{
+					Count = Controller.GetSharedReferenceCount();
+					AutoRTFM::AbortTransaction();
+				});
+
+			REQUIRE(AutoRTFM::ETransactionResult::AbortedByRequest == Result);
+			REQUIRE(0 == Count);
+		}
+
+		SECTION("With Commit")
+		{
+			int32 Count = 0;
+
+			AutoRTFM::ETransactionResult Result = AutoRTFM::Transact([&]()
+				{
+					Count = Controller.GetSharedReferenceCount();
+				});
+
+			REQUIRE(AutoRTFM::ETransactionResult::Committed == Result);
+			REQUIRE(1 == Count);
+		}
+	}
+
+	SECTION("IsUnique")
+	{
+		SharedPointerInternals::TIntrusiveReferenceController<int, ESPMode::ThreadSafe> Controller(42);
+
+		SECTION("True")
+		{
+			bool Unique = false;
+
+			AutoRTFM::ETransactionResult Result = AutoRTFM::Transact([&]()
+				{
+					Unique = Controller.IsUnique();
+				});
+
+			REQUIRE(AutoRTFM::ETransactionResult::Committed == Result);
+			REQUIRE(Unique);
+		}
+
+		SECTION("False")
+		{
+			// Add a count to make us not unique.
+			Controller.AddSharedReference();
+
+			bool Unique = true;
+
+			AutoRTFM::ETransactionResult Result = AutoRTFM::Transact([&]()
+				{
+					Unique = Controller.IsUnique();
+				});
+
+			REQUIRE(AutoRTFM::ETransactionResult::Committed == Result);
+			REQUIRE(!Unique);
+		}
+	}
+
+	SECTION("ReleaseSharedReference")
+	{
+		SharedPointerInternals::TIntrusiveReferenceController<int, ESPMode::ThreadSafe> Controller(42);
+
+		// We don't want the add weak reference deleter to trigger in this test so add another to its count.
+		Controller.AddWeakReference();
+
+		SECTION("With Abort")
+		{
+			AutoRTFM::ETransactionResult Result = AutoRTFM::Transact([&]()
+				{
+					Controller.ReleaseSharedReference();
+					AutoRTFM::AbortTransaction();
+				});
+
+			REQUIRE(AutoRTFM::ETransactionResult::AbortedByRequest == Result);
+			REQUIRE(1 == Controller.GetSharedReferenceCount());
+		}
+
+		SECTION("With Commit")
+		{
+			AutoRTFM::ETransactionResult Result = AutoRTFM::Transact([&]()
+				{
+					Controller.ReleaseSharedReference();
+				});
+
+			REQUIRE(AutoRTFM::ETransactionResult::Committed == Result);
+		}
+	}
+
+	SECTION("ReleaseWeakReference")
+	{
+		auto* Controller = new SharedPointerInternals::TIntrusiveReferenceController<int, ESPMode::ThreadSafe>(42);
+
+		SECTION("With Abort")
+		{
+			AutoRTFM::ETransactionResult Result = AutoRTFM::Transact([&]()
+				{
+					Controller->ReleaseWeakReference();
+					AutoRTFM::AbortTransaction();
+				});
+
+			REQUIRE(AutoRTFM::ETransactionResult::AbortedByRequest == Result);
+			REQUIRE(1 == Controller->WeakReferenceCount);
+		}
+
+		SECTION("With Commit")
+		{
+			AutoRTFM::ETransactionResult Result = AutoRTFM::Transact([&]()
+				{
+					Controller->ReleaseWeakReference();
+				});
+
+			REQUIRE(AutoRTFM::ETransactionResult::Committed == Result);
+		}
+	}
+
+	SECTION("GetObjectPtr")
+	{
+		SharedPointerInternals::TIntrusiveReferenceController<int, ESPMode::ThreadSafe> Controller(42);
+
+		SECTION("With Abort")
+		{
+			AutoRTFM::ETransactionResult Result = AutoRTFM::Transact([&]()
+				{
+					*Controller.GetObjectPtr() = 13;
+					AutoRTFM::AbortTransaction();
+				});
+
+			REQUIRE(AutoRTFM::ETransactionResult::AbortedByRequest == Result);
+			REQUIRE(42 == *Controller.GetObjectPtr());
+		}
+
+		SECTION("With Commit")
+		{
+			AutoRTFM::ETransactionResult Result = AutoRTFM::Transact([&]()
+				{
+					*Controller.GetObjectPtr() = 13;
+				});
+
+			REQUIRE(AutoRTFM::ETransactionResult::Committed == Result);
+			REQUIRE(13 == *Controller.GetObjectPtr());
+		}
+	}
+}
diff --git a/Engine/Source/Runtime/Core/Public/Templates/SharedPointerInternals.h b/Engine/Source/Runtime/Core/Public/Templates/SharedPointerInternals.h
--- a/Engine/Source/Runtime/Core/Public/Templates/SharedPointerInternals.h
+++ b/Engine/Source/Runtime/Core/Public/Templates/SharedPointerInternals.h
@@ -79,8 +79,16 @@ namespace SharedPointerInternals
 				//
 				// This is equivalent to https://en.cppreference.com/w/cpp/memory/shared_ptr/use_count
 
-				// This reference count may be accessed by multiple threads
-				return SharedReferenceCount.load(std::memory_order_relaxed);
+				int32 Count = 0;
+
+				UE_AUTORTFM_OPEN(
+					{
+						// This reference count may be accessed by multiple threads
+						Count = SharedReferenceCount.load(std::memory_order_relaxed);
+					}
+				);
+
+				return Count;
 			}
 			else
 			{
@@ -91,19 +99,7 @@ namespace SharedPointerInternals
 		/** Checks if there is exactly one reference left to the object. */
 		FORCEINLINE bool IsUnique() const
 		{
-			if constexpr (Mode == ESPMode::ThreadSafe)
-			{
-				// This is equivalent to https://en.cppreference.com/w/cpp/memory/shared_ptr/unique,
-				// however instead of deprecating it, we implement it with an acquire, which should
-				// suit our use cases.
-
-				// This reference count may be accessed by multiple threads
-				return SharedReferenceCount.load(std::memory_order_acquire) == 1;
-			}
-			else
-			{
-				return SharedReferenceCount == 1;
-			}
+			return 1 == GetSharedReferenceCount();
 		}
 
 		/** Adds a shared reference to this counter */
@@ -188,7 +184,7 @@ namespace SharedPointerInternals
 					}
 				});
 
-				// If we succeedd in taking a shared reference count, we need to undo that on an abort.
+				// If we succeeded in taking a shared reference count, we need to undo that on an abort.
 				if (bSucceeded)
 				{
 					UE_AUTORTFM_ONABORT(
@@ -217,7 +213,7 @@ namespace SharedPointerInternals
 		{
 			if constexpr (Mode == ESPMode::ThreadSafe)
 			{
-				UE_AUTORTFM_ONCOMMIT(
+				AutoRTFM::OnCommit([this]
 				{
 					// std::memory_order_acq_rel is used here so that, if we do end up executing the destructor, it's not possible
 					// for side effects from executing the destructor end up being visible before we've determined that the shared
@@ -260,12 +256,24 @@ namespace SharedPointerInternals
 				// See AddSharedReference for the same reasons that std::memory_order_relaxed is used in this function.
 
 #if defined(_MSC_VER) && (defined(_M_X64) || defined(_M_IX86))
-				// We do a regular SC increment here because it maps to an _InterlockedIncrement (lock inc).
-				// The codegen for a relaxed fetch_add is actually much worse under MSVC (lock xadd).
-				++WeakReferenceCount;
+				UE_AUTORTFM_OPEN(
+					{
+						// We do a regular SC increment here because it maps to an _InterlockedIncrement (lock inc).
+						// The codegen for a relaxed fetch_add is actually much worse under MSVC (lock xadd).
+						++WeakReferenceCount;
+					});
 #else
-				WeakReferenceCount.fetch_add(1, std::memory_order_relaxed);
+				UE_AUTORTFM_OPEN(
+					{
+						WeakReferenceCount.fetch_add(1, std::memory_order_relaxed);
+					});
 #endif
+
+				// If the transaction would abort, we need to undo adding the reference.
+				UE_AUTORTFM_ONABORT(
+					{
+						ReleaseWeakReference();
+					});
 			}
 			else
 			{
@@ -278,20 +286,23 @@ namespace SharedPointerInternals
 		{
 			if constexpr (Mode == ESPMode::ThreadSafe)
 			{
-				// See ReleaseSharedReference for the same reasons that std::memory_order_acq_rel is used in this function.
+				AutoRTFM::OnCommit([this]
+					{
+						// See ReleaseSharedReference for the same reasons that std::memory_order_acq_rel is used in this function.
 
-				int32 OldWeakCount = WeakReferenceCount.fetch_sub(1, std::memory_order_acq_rel);
-				checkSlow(OldWeakCount > 0);
-				if (OldWeakCount == 1)
-				{
-					// Disable this if running clang's static analyzer. Passing shared pointers
-					// and references to functions it cannot reason about, produces false
-					// positives about use-after-free in the TSharedPtr/TSharedRef destructors.
+						int32 OldWeakCount = WeakReferenceCount.fetch_sub(1, std::memory_order_acq_rel);
+						checkSlow(OldWeakCount > 0);
+						if (OldWeakCount == 1)
+						{
+							// Disable this if running clang's static analyzer. Passing shared pointers
+							// and references to functions it cannot reason about, produces false
+							// positives about use-after-free in the TSharedPtr/TSharedRef destructors.
 #if !defined(__clang_analyzer__)
-					// No more references to this reference count.  Destroy it!
-					delete this;
+							// No more references to this reference count.  Destroy it!
+							delete this;
 #endif
-				}
+						}
+					});
 			}
 			else
 			{
