Remove the older and unused pragma approach to enabling AutoRTFM.

This pragma didn't work with our unity builds, and we moved away from it anyway.

diff --git a/Engine/Source/Programs/AutoRTFMTests/Private/API.cpp b/Engine/Source/Programs/AutoRTFMTests/Private/API.cpp
--- a/Engine/Source/Programs/AutoRTFMTests/Private/API.cpp
+++ b/Engine/Source/Programs/AutoRTFMTests/Private/API.cpp
@@ -1,7 +1,5 @@
 // Copyright Epic Games, Inc. All Rights Reserved.
 
-#pragma autortfm
-
 #include "API.h"
 #include "Catch2Includes.h"
 
diff --git a/Engine/Source/Programs/AutoRTFMTests/Private/APINoAutoRTFM.cpp b/Engine/Source/Programs/AutoRTFMTests/Private/APINoAutoRTFM.cpp
--- a/Engine/Source/Programs/AutoRTFMTests/Private/APINoAutoRTFM.cpp
+++ b/Engine/Source/Programs/AutoRTFMTests/Private/APINoAutoRTFM.cpp
@@ -1,24 +1,24 @@
 // Copyright Epic Games, Inc. All Rights Reserved.
 
 #include "API.h"
+#include <AutoRTFM/AutoRTFM.h>
 
-int NoAutoRTFM::DoSomethingC(int I)
+UE_AUTORTFM_NOAUTORTFM int NoAutoRTFM::DoSomethingC(int I)
 {
     return I + 13;
 }
 
-int NoAutoRTFM::DoSomethingInTransactionC(int I, void*)
+UE_AUTORTFM_NOAUTORTFM int NoAutoRTFM::DoSomethingInTransactionC(int I, void*)
 {
     return I + 42;
 }
 
-int NoAutoRTFM::DoSomethingCpp(int I)
+UE_AUTORTFM_NOAUTORTFM int NoAutoRTFM::DoSomethingCpp(int I)
 {
     return I + 13;
 }
 
-int NoAutoRTFM::DoSomethingInTransactionCpp(int I, void*)
+UE_AUTORTFM_NOAUTORTFM int NoAutoRTFM::DoSomethingInTransactionCpp(int I, void*)
 {
     return I + 42;
 }
-
diff --git a/Engine/Source/Programs/AutoRTFMTests/Private/Abort.cpp b/Engine/Source/Programs/AutoRTFMTests/Private/Abort.cpp
--- a/Engine/Source/Programs/AutoRTFMTests/Private/Abort.cpp
+++ b/Engine/Source/Programs/AutoRTFMTests/Private/Abort.cpp
@@ -1,7 +1,5 @@
 // Copyright Epic Games, Inc. All Rights Reserved.
 
-#pragma autortfm
-
 #include "Catch2Includes.h"
 #include <AutoRTFM/AutoRTFM.h>
 #include <map>
diff --git a/Engine/Source/Programs/AutoRTFMTests/Private/Allocators.cpp b/Engine/Source/Programs/AutoRTFMTests/Private/Allocators.cpp
--- a/Engine/Source/Programs/AutoRTFMTests/Private/Allocators.cpp
+++ b/Engine/Source/Programs/AutoRTFMTests/Private/Allocators.cpp
@@ -1,7 +1,5 @@
 // Copyright Epic Games, Inc. All Rights Reserved.
 
-#pragma autortfm
-
 #include "Catch2Includes.h"
 #include <AutoRTFM/AutoRTFM.h>
 #include <new>
diff --git a/Engine/Source/Programs/AutoRTFMTests/Private/Benchmarks.cpp b/Engine/Source/Programs/AutoRTFMTests/Private/Benchmarks.cpp
--- a/Engine/Source/Programs/AutoRTFMTests/Private/Benchmarks.cpp
+++ b/Engine/Source/Programs/AutoRTFMTests/Private/Benchmarks.cpp
@@ -1,7 +1,5 @@
 // Copyright Epic Games, Inc. All Rights Reserved.
 
-#pragma autortfm
-
 #include "Catch2Includes.h"
 #include <AutoRTFM/AutoRTFM.h>
 #include <thread>
diff --git a/Engine/Source/Programs/AutoRTFMTests/Private/CString.cpp b/Engine/Source/Programs/AutoRTFMTests/Private/CString.cpp
--- a/Engine/Source/Programs/AutoRTFMTests/Private/CString.cpp
+++ b/Engine/Source/Programs/AutoRTFMTests/Private/CString.cpp
@@ -1,7 +1,5 @@
 // Copyright Epic Games, Inc. All Rights Reserved.
 
-#pragma autortfm
-
 #include "Catch2Includes.h"
 #include <AutoRTFM/AutoRTFM.h>
 #include <cstring>
diff --git a/Engine/Source/Programs/AutoRTFMTests/Private/Commit.cpp b/Engine/Source/Programs/AutoRTFMTests/Private/Commit.cpp
--- a/Engine/Source/Programs/AutoRTFMTests/Private/Commit.cpp
+++ b/Engine/Source/Programs/AutoRTFMTests/Private/Commit.cpp
@@ -1,7 +1,5 @@
 // Copyright Epic Games, Inc. All Rights Reserved.
 
-#pragma autortfm
-
 #include "Catch2Includes.h"
 #include <AutoRTFM/AutoRTFM.h>
 #include <map>
diff --git a/Engine/Source/Programs/AutoRTFMTests/Private/CopyAddInt64.cpp b/Engine/Source/Programs/AutoRTFMTests/Private/CopyAddInt64.cpp
--- a/Engine/Source/Programs/AutoRTFMTests/Private/CopyAddInt64.cpp
+++ b/Engine/Source/Programs/AutoRTFMTests/Private/CopyAddInt64.cpp
@@ -1,7 +1,5 @@
 // Copyright Epic Games, Inc. All Rights Reserved.
 
-#pragma autortfm
-
 #include "Catch2Includes.h"
 #include <AutoRTFM/AutoRTFM.h>
 
@@ -11,4 +9,3 @@ TEST_CASE("CopyAddInt64")
     AutoRTFM::Commit([&] () { X += 0xa1b2c3d4e5f60789llu; });
     REQUIRE(X == 0x123456789abcdef0llu + 0xa1b2c3d4e5f60789llu);
 }
-
diff --git a/Engine/Source/Programs/AutoRTFMTests/Private/FakeVM.cpp b/Engine/Source/Programs/AutoRTFMTests/Private/FakeVM.cpp
--- a/Engine/Source/Programs/AutoRTFMTests/Private/FakeVM.cpp
+++ b/Engine/Source/Programs/AutoRTFMTests/Private/FakeVM.cpp
@@ -1,7 +1,5 @@
 // Copyright Epic Games, Inc. All Rights Reserved.
 
-#pragma autortfm
-
 #include "Catch2Includes.h"
 #include <AutoRTFM/AutoRTFM.h>
 #include <vector>
diff --git a/Engine/Source/Programs/AutoRTFMTests/Private/GraphBenchmark.cpp b/Engine/Source/Programs/AutoRTFMTests/Private/GraphBenchmark.cpp
--- a/Engine/Source/Programs/AutoRTFMTests/Private/GraphBenchmark.cpp
+++ b/Engine/Source/Programs/AutoRTFMTests/Private/GraphBenchmark.cpp
@@ -1,7 +1,5 @@
 // Copyright Epic Games, Inc. All Rights Reserved.
 
-#pragma autortfm
-
 #include "Catch2Includes.h"
 #include <AutoRTFM/AutoRTFM.h>
 #include <vector>
diff --git a/Engine/Source/Programs/AutoRTFMTests/Private/Open.cpp b/Engine/Source/Programs/AutoRTFMTests/Private/Open.cpp
--- a/Engine/Source/Programs/AutoRTFMTests/Private/Open.cpp
+++ b/Engine/Source/Programs/AutoRTFMTests/Private/Open.cpp
@@ -1,7 +1,5 @@
 // Copyright Epic Games, Inc. All Rights Reserved.
 
-#pragma autortfm
-
 #include "Catch2Includes.h"
 #include <AutoRTFM/AutoRTFM.h>
 #include <atomic>
diff --git a/Engine/Source/Programs/AutoRTFMTests/Private/OpenAPI.cpp b/Engine/Source/Programs/AutoRTFMTests/Private/OpenAPI.cpp
--- a/Engine/Source/Programs/AutoRTFMTests/Private/OpenAPI.cpp
+++ b/Engine/Source/Programs/AutoRTFMTests/Private/OpenAPI.cpp
@@ -1,7 +1,5 @@
 // Copyright Epic Games, Inc. All Rights Reserved.
 
-#pragma autortfm
-
 #include "Catch2Includes.h"
 #include <AutoRTFM/AutoRTFM.h>
 
diff --git a/Engine/Source/Programs/AutoRTFMTests/Private/PackedPointer.cpp b/Engine/Source/Programs/AutoRTFMTests/Private/PackedPointer.cpp
--- a/Engine/Source/Programs/AutoRTFMTests/Private/PackedPointer.cpp
+++ b/Engine/Source/Programs/AutoRTFMTests/Private/PackedPointer.cpp
@@ -1,7 +1,5 @@
 // Copyright Epic Games, Inc. All Rights Reserved.
 
-#pragma autortfm
-
 #include "Catch2Includes.h"
 #include <AutoRTFM/AutoRTFM.h>
 #include <inttypes.h>
diff --git a/Engine/Source/Programs/AutoRTFMTests/Private/SharedPointerTests.cpp b/Engine/Source/Programs/AutoRTFMTests/Private/SharedPointerTests.cpp
--- a/Engine/Source/Programs/AutoRTFMTests/Private/SharedPointerTests.cpp
+++ b/Engine/Source/Programs/AutoRTFMTests/Private/SharedPointerTests.cpp
@@ -1,7 +1,5 @@
 // Copyright Epic Games, Inc. All Rights Reserved.
 
-#pragma autortfm
-
 #include "Catch2Includes.h"
 #include <AutoRTFM/AutoRTFM.h>
 #include "Templates/SharedPointer.h"
diff --git a/Engine/Source/Programs/AutoRTFMTests/Private/Statics.cpp b/Engine/Source/Programs/AutoRTFMTests/Private/Statics.cpp
--- a/Engine/Source/Programs/AutoRTFMTests/Private/Statics.cpp
+++ b/Engine/Source/Programs/AutoRTFMTests/Private/Statics.cpp
@@ -1,7 +1,5 @@
 // Copyright Epic Games, Inc. All Rights Reserved.
 
-#pragma autortfm
-
 #include "Catch2Includes.h"
 #include <AutoRTFM/AutoRTFM.h>
 
diff --git a/Engine/Source/Programs/AutoRTFMTests/Private/UECoreTests.cpp b/Engine/Source/Programs/AutoRTFMTests/Private/UECoreTests.cpp
--- a/Engine/Source/Programs/AutoRTFMTests/Private/UECoreTests.cpp
+++ b/Engine/Source/Programs/AutoRTFMTests/Private/UECoreTests.cpp
@@ -1,7 +1,5 @@
 // Copyright Epic Games, Inc. All Rights Reserved.
 
-#pragma autortfm
-
 #include "Catch2Includes.h"
 #include "AutoRTFM/AutoRTFM.h"
 #include "Delegates/IDelegateInstance.h"
diff --git a/Engine/Source/Programs/AutoRTFMTests/Private/UObjectTests.cpp b/Engine/Source/Programs/AutoRTFMTests/Private/UObjectTests.cpp
--- a/Engine/Source/Programs/AutoRTFMTests/Private/UObjectTests.cpp
+++ b/Engine/Source/Programs/AutoRTFMTests/Private/UObjectTests.cpp
@@ -1,7 +1,5 @@
 // Copyright Epic Games, Inc. All Rights Reserved.
 
-#pragma autortfm
-
 #include "MyAutoRTFMTestObject.h"
 #include "Catch2Includes.h"
 #include <AutoRTFM/AutoRTFM.h>
diff --git a/Engine/Source/Programs/AutoRTFMTests/Private/UnorderedSet.cpp b/Engine/Source/Programs/AutoRTFMTests/Private/UnorderedSet.cpp
--- a/Engine/Source/Programs/AutoRTFMTests/Private/UnorderedSet.cpp
+++ b/Engine/Source/Programs/AutoRTFMTests/Private/UnorderedSet.cpp
@@ -1,7 +1,5 @@
 // Copyright Epic Games, Inc. All Rights Reserved.
 
-#pragma autortfm
-
 #include "Catch2Includes.h"
 #include <AutoRTFM/AutoRTFM.h>
 #include <unordered_set>
