Make PopOnAbortHandler only pop the last pushed for that key.

diff --git a/Engine/Source/Programs/AutoRTFMTests/Private/Abort.cpp b/Engine/Source/Programs/AutoRTFMTests/Private/Abort.cpp
--- a/Engine/Source/Programs/AutoRTFMTests/Private/Abort.cpp
+++ b/Engine/Source/Programs/AutoRTFMTests/Private/Abort.cpp
@@ -292,7 +292,9 @@ TEST_CASE("Abort.PushOnAbortHandler_Duplicates1")
 	});
 
 	REQUIRE(AutoRTFM::ETransactionResult::AbortedByRequest == Result);
-	REQUIRE(Value == 55);
+
+	// The first push on abort will still go through.
+	REQUIRE(Value == 77);
 }
 
 TEST_CASE("Abort.PushOnAbortHandler_Duplicates2")
@@ -302,8 +304,8 @@ TEST_CASE("Abort.PushOnAbortHandler_Duplicates2")
 	const AutoRTFM::ETransactionResult Result = AutoRTFM::Transact([&]
 	{
 		Value = 66;
-		AutoRTFM::PushOnAbortHandler(UIntToPointer(747), [&Value](){ Value = 77; });
-		AutoRTFM::PushOnAbortHandler(UIntToPointer(747), [&Value](){ Value = 88; });
+		AutoRTFM::PushOnAbortHandler(UIntToPointer(747), [&Value](){ Value += 12; });
+		AutoRTFM::PushOnAbortHandler(UIntToPointer(747), [&Value](){ Value = 65; });
 		Value = 99;
 
 		AutoRTFM::AbortTransaction();
@@ -313,22 +315,129 @@ TEST_CASE("Abort.PushOnAbortHandler_Duplicates2")
 	REQUIRE(Value == 77);
 }
 
-TEST_CASE("Abort.PushOnAbortHandler_MultiplePops")
+TEST_CASE("Abort.PushOnAbortHandler_Order")
 {
-	int Value = 55;
+	SECTION("HandlerSandwich")
+	{
+		SECTION("WithoutPop")
+		{
+			int Value = 37;
 
-	const AutoRTFM::ETransactionResult Result = AutoRTFM::Transact([&]
+			const AutoRTFM::ETransactionResult Result = AutoRTFM::Transact([&]
+				{
+					AutoRTFM::OnAbort([&Value] { REQUIRE(42 == Value); Value += 1; });
+					AutoRTFM::PushOnAbortHandler(UIntToPointer(747), [&Value]() { REQUIRE(40 == Value); Value += 2; });
+					AutoRTFM::OnAbort([&Value] { REQUIRE(37 == Value); Value += 3; });
+
+					Value = 99;
+
+					AutoRTFM::AbortTransaction();
+				});
+
+			REQUIRE(AutoRTFM::ETransactionResult::AbortedByRequest == Result);
+			REQUIRE(Value == 43);
+		}
+
+		SECTION("WithPop")
+		{
+			int Value = 37;
+
+			const AutoRTFM::ETransactionResult Result = AutoRTFM::Transact([&]
+				{
+					AutoRTFM::OnAbort([&Value] { REQUIRE(40 == Value); Value += 1; });
+					AutoRTFM::PushOnAbortHandler(UIntToPointer(747), [&Value]() { REQUIRE(false); });
+					AutoRTFM::OnAbort([&Value] { REQUIRE(37 == Value); Value += 3; });
+
+					AutoRTFM::PopOnAbortHandler(UIntToPointer(747));
+
+					Value = 99;
+
+					AutoRTFM::AbortTransaction();
+				});
+
+			REQUIRE(AutoRTFM::ETransactionResult::AbortedByRequest == Result);
+			REQUIRE(Value == 41);
+		}
+	}
+
+	SECTION("HandlerInChild")
 	{
-		Value = 66;
-		AutoRTFM::PopOnAbortHandler(UIntToPointer(747));
-		AutoRTFM::PopOnAbortHandler(UIntToPointer(747));
-		AutoRTFM::PopOnAbortHandler(UIntToPointer(747));
-		AutoRTFM::PopOnAbortHandler(UIntToPointer(747));
-		Value = 99;
+		SECTION("WithoutPop")
+		{
+			int Value = 37;
 
-		AutoRTFM::AbortTransaction();
-	});
+			const AutoRTFM::ETransactionResult Result = AutoRTFM::Transact([&]
+				{
+					AutoRTFM::OnAbort([&Value] { REQUIRE(42 == Value); Value += 1; });
 
-	REQUIRE(AutoRTFM::ETransactionResult::AbortedByRequest == Result);
-	REQUIRE(Value == 55);
+					// Make a child transaction.
+					AutoRTFM::Commit([&]
+						{
+							AutoRTFM::PushOnAbortHandler(UIntToPointer(747), [&Value]() { REQUIRE(40 == Value); Value += 2; });
+						});
+
+					AutoRTFM::OnAbort([&Value] { REQUIRE(37 == Value); Value += 3; });
+
+					Value = 99;
+
+					AutoRTFM::AbortTransaction();
+				});
+
+			REQUIRE(AutoRTFM::ETransactionResult::AbortedByRequest == Result);
+			REQUIRE(Value == 43);
+		}
+
+		SECTION("WithPop")
+		{
+			int Value = 37;
+
+			const AutoRTFM::ETransactionResult Result = AutoRTFM::Transact([&]
+				{
+					AutoRTFM::OnAbort([&Value] { REQUIRE(40 == Value); Value += 1; });
+					
+					// Make a child transaction.
+					AutoRTFM::Commit([&]
+						{
+							AutoRTFM::PushOnAbortHandler(UIntToPointer(747), [&Value]() { REQUIRE(false); });
+						});
+
+					AutoRTFM::OnAbort([&Value] { REQUIRE(37 == Value); Value += 3; });
+
+					// Bit funky, but we can pop the child's push here!
+					AutoRTFM::PopOnAbortHandler(UIntToPointer(747));
+
+					Value = 99;
+
+					AutoRTFM::AbortTransaction();
+				});
+
+			REQUIRE(AutoRTFM::ETransactionResult::AbortedByRequest == Result);
+			REQUIRE(Value == 41);
+		}
+
+		SECTION("AbortInChild")
+		{
+			int Value = 99;
+
+			AutoRTFM::ETransactionResult Result = AutoRTFM::ETransactionResult::Committed;
+			AutoRTFM::Commit([&]
+				{
+					AutoRTFM::OnCommit([&Value] { REQUIRE(37 == Value); Value += 1; });
+
+					// Make a child transaction.
+					Result = AutoRTFM::Transact([&]
+						{
+							AutoRTFM::PushOnAbortHandler(UIntToPointer(747), [&Value]() { REQUIRE(38 == Value); Value += 2; });
+							AutoRTFM::AbortTransaction();
+						});
+
+					AutoRTFM::OnCommit([&Value] { REQUIRE(40 == Value); Value += 3; });
+
+					Value = 37;
+				});
+
+			REQUIRE(AutoRTFM::ETransactionResult::AbortedByRequest == Result);
+			REQUIRE(Value == 43);
+		}
+	}
 }
diff --git a/Engine/Source/Runtime/Core/Private/AutoRTFM/API.cpp b/Engine/Source/Runtime/Core/Private/AutoRTFM/API.cpp
--- a/Engine/Source/Runtime/Core/Private/AutoRTFM/API.cpp
+++ b/Engine/Source/Runtime/Core/Private/AutoRTFM/API.cpp
@@ -496,7 +496,8 @@ extern "C" UE_AUTORTFM_NOAUTORTFM void RTFM_PopOnAbortHandler(const void* Key)
 {
 	FContext* Context = FContext::Get();
     ASSERT(Context->GetStatus() == EContextStatus::OnTrack);
-    Context->GetCurrentTransaction()->PopDeferUntilAbortHandler(Key);
+    const bool bResult = Context->GetCurrentTransaction()->PopDeferUntilAbortHandler(Key);
+	ASSERT(bResult);
 }
 
 extern "C" UE_AUTORTFM_NOAUTORTFM void RTFM_autortfm_on_commit(void (*Work)(void*), void* Arg)
diff --git a/Engine/Source/Runtime/Core/Private/AutoRTFM/TaskArray.h b/Engine/Source/Runtime/Core/Private/AutoRTFM/TaskArray.h
--- a/Engine/Source/Runtime/Core/Private/AutoRTFM/TaskArray.h
+++ b/Engine/Source/Runtime/Core/Private/AutoRTFM/TaskArray.h
@@ -47,52 +47,70 @@ public:
 		Latest.Push(TTaskArrayKeyValuePair(SKey{}, MoveTemp(Value)));
     }
 
+	void Add(const InElementType& Value)
+	{
+		Latest.Push(TTaskArrayKeyValuePair(SKey{}, Value));
+	}
+
 	void AddKeyed(SKey Key, InElementType&& Value)
     {
 		ASSERT(Key != SKey{});
 		Latest.Push(TTaskArrayKeyValuePair(Key, MoveTemp(Value)));
     }
 
-    void Add(const InElementType& Value)
-    {
-		Latest.Push(TTaskArrayKeyValuePair(SKey{}, Value));
-    }
-
 	void AddKeyed(SKey Key, const InElementType& Value)
     {
 		ASSERT(Key != SKey{});
 		Latest.Push(TTaskArrayKeyValuePair(Key, Value));
     }
 
-	void DeleteKey(SKey Key)
+	bool DeleteKey(SKey Key)
 	{
 		ASSERT(Key != SKey{});
 
-		auto EraseKeyFromArray = [Key](FInternalArray& Array) -> void
+		auto EraseLastAddedKeyFromArray = [Key](FInternalArray& Array) -> bool
 		{
-			uint64 WriteIdx = 0;
-			for(uint64 ReadIdx = 0; ReadIdx < Array.Num(); ReadIdx++)
+			for(uint64 Idx = 0; Idx < Array.Num(); Idx++)
 			{
-				if(Array[ReadIdx].Key != Key)
+				uint64 BackwardsIdx = Array.Num() - 1 - Idx;
+
+				if(Array[BackwardsIdx].Key != Key)
 				{
-					// keep it
-					if(WriteIdx != ReadIdx) // the common case will be WriteIdx==ReadIdx so don't invoke an unnecessary copy-to-self
-					{
-						Array[WriteIdx] = MoveTemp(Array[ReadIdx]);
-					}
-					WriteIdx++;
+					continue;
 				}
+
+				// We found our key to erase! Nuke it.
+
+				// We start at the element just after our element.
+				BackwardsIdx++;
+
+				for (; BackwardsIdx < Array.Num(); BackwardsIdx++)
+				{
+					Array[BackwardsIdx - 1] = MoveTemp(Array[BackwardsIdx]);
+				}
+
+				Array.SetNum(Array.Num() - 1, EAllowShrinking::No);
+
+				return true;
 			}
 
-			Array.SetNum(WriteIdx, EAllowShrinking::No);
+			return false;
 		};
 
-		for(FInternalArray& Array : Stash)
+		if (EraseLastAddedKeyFromArray(Latest))
 		{
-			EraseKeyFromArray(Array);
+			return true;
+		}
+
+		for (FInternalArray& StashedVectorBox : TBackwards(Stash))
+		{
+			if (EraseLastAddedKeyFromArray(StashedVectorBox))
+			{
+				return true;
+			}
 		}
 
-		EraseKeyFromArray(Latest);
+		return false;
 	}
 
     void AddAll(TTaskArray&& Other)
diff --git a/Engine/Source/Runtime/Core/Private/AutoRTFM/Transaction.h b/Engine/Source/Runtime/Core/Private/AutoRTFM/Transaction.h
--- a/Engine/Source/Runtime/Core/Private/AutoRTFM/Transaction.h
+++ b/Engine/Source/Runtime/Core/Private/AutoRTFM/Transaction.h
@@ -64,7 +64,7 @@ public:
     void DeferUntilCommit(TFunction<void()>&&);
     void DeferUntilAbort(TFunction<void()>&&);
     void PushDeferUntilAbortHandler(const void* Key, TFunction<void()>&&);
-    void PopDeferUntilAbortHandler(const void* Key);
+    bool PopDeferUntilAbortHandler(const void* Key);
 
     void AbortAndThrow();
     void AbortWithoutThrowing();
diff --git a/Engine/Source/Runtime/Core/Private/AutoRTFM/TransactionInlines.h b/Engine/Source/Runtime/Core/Private/AutoRTFM/TransactionInlines.h
--- a/Engine/Source/Runtime/Core/Private/AutoRTFM/TransactionInlines.h
+++ b/Engine/Source/Runtime/Core/Private/AutoRTFM/TransactionInlines.h
@@ -167,9 +167,9 @@ UE_AUTORTFM_FORCEINLINE void FTransaction::PushDeferUntilAbortHandler(const void
     AbortTasks.AddKeyed(Key, MoveTemp(Copy));
 }
 
-UE_AUTORTFM_FORCEINLINE void FTransaction::PopDeferUntilAbortHandler(const void* Key)
+UE_AUTORTFM_FORCEINLINE bool FTransaction::PopDeferUntilAbortHandler(const void* Key)
 {
-    AbortTasks.DeleteKey(Key);
+    return AbortTasks.DeleteKey(Key);
 }
 
 UE_AUTORTFM_FORCEINLINE void FTransaction::CollectStats() const
