Add some function pointer AutoRTFM tests.

diff --git a/Engine/Source/Programs/AutoRTFMTests/Private/FunctionPointerTests.cpp b/Engine/Source/Programs/AutoRTFMTests/Private/FunctionPointerTests.cpp
new file mode 100644
--- /dev/null
+++ b/Engine/Source/Programs/AutoRTFMTests/Private/FunctionPointerTests.cpp
@@ -0,0 +1,141 @@
+// Copyright Epic Games, Inc. All Rights Reserved.
+
+#include "Catch2Includes.h"
+#include <AutoRTFM/AutoRTFM.h>
+#include "Templates/Function.h"
+
+namespace
+{
+	int Something()
+	{
+		if (AutoRTFM::IsClosed())
+		{
+			return 42;
+		}
+		else
+		{
+			AutoRTFM::AbortTransaction();
+		}
+	}
+
+	typedef int (*CStyleType)();
+
+	UE_DISABLE_OPTIMIZATION_SHIP
+	CStyleType GetSomething()
+	{
+		return &Something;
+	}
+
+	UE_ENABLE_OPTIMIZATION_SHIP
+}
+
+TEST_CASE("FunctionPointer.CStyle")
+{
+	int Result = 0;
+	AutoRTFM::Commit([&]
+		{
+			CStyleType CStyle = GetSomething();
+			Result = CStyle();
+		});
+
+	REQUIRE(42 == Result);
+}
+
+TEST_CASE("FunctionPointer.TFunction")
+{
+	SECTION("Created inside transaction")
+	{
+		int Result = 0;
+		AutoRTFM::Commit([&]
+			{
+				TFunction<void()> MyFunc = [&Result]() -> void
+					{
+						Result = 42;
+					};
+
+				if (MyFunc)
+				{
+					MyFunc();
+				}
+
+				MyFunc.CheckCallable();
+
+				MyFunc.Reset();
+			});
+
+		REQUIRE(42 == Result);
+	}
+
+	SECTION("Created outside transaction")
+	{
+		int Result = 0;
+		TFunction<void()> MyFunc = [&Result]() -> void
+			{
+				Result = 42;
+			};
+
+		AutoRTFM::Commit([&]
+			{
+				if (MyFunc)
+				{
+					MyFunc();
+				}
+
+				MyFunc.CheckCallable();
+
+				MyFunc.Reset();
+			});
+
+		REQUIRE(42 == Result);
+
+	}
+}
+
+TEST_CASE("FunctionPointer.TUniqueFunction")
+{
+	SECTION("Created inside transaction")
+	{
+		int Result = 0;
+		AutoRTFM::Commit([&]
+			{
+				TUniqueFunction<void()> MyFunc = [&Result]() -> void
+					{
+						Result = 42;
+					};
+
+				if (MyFunc)
+				{
+					MyFunc();
+				}
+
+				MyFunc.CheckCallable();
+
+				MyFunc.Reset();
+			});
+
+		REQUIRE(42 == Result);
+	}
+
+	SECTION("Created outside transaction")
+	{
+		int Result = 0;
+		TUniqueFunction<void()> MyFunc = [&Result]() -> void
+			{
+				Result = 42;
+			};
+
+		AutoRTFM::Commit([&]
+			{
+				if (MyFunc)
+				{
+					MyFunc();
+				}
+
+				MyFunc.CheckCallable();
+
+				MyFunc.Reset();
+			});
+
+		REQUIRE(42 == Result);
+	}
+}
