Make AutoRTFMTests 1.2x faster to run by removing/reducing what we run.

diff --git a/Engine/Source/Programs/AutoRTFMTests/Private/Benchmarks.cpp b/Engine/Source/Programs/AutoRTFMTests/Private/Benchmarks.cpp
--- a/Engine/Source/Programs/AutoRTFMTests/Private/Benchmarks.cpp
+++ b/Engine/Source/Programs/AutoRTFMTests/Private/Benchmarks.cpp
@@ -2,59 +2,327 @@
 
 #include "Catch2Includes.h"
 #include <AutoRTFM/AutoRTFM.h>
-#include <thread>
 #include <vector>
+#include <unordered_set>
 
-struct MyStruct final {
-	uint8_t u8;
-	uint64_t u64;
+using namespace std;
 
-	void DoSomething() {
-		u8++;
-		u64++;
+template<typename T>
+class FNode
+{
+public:
+	FNode(T Value) : Value(Value) { }
+
+	T GetValue() const { return Value; }
+
+	void AddEdge(FNode* Target)
+	{
+		Edges.push_back(Target);
 	}
+
+	const vector<FNode*>& GetEdges() const { return Edges; }
+
+private:
+	vector<FNode*> Edges;
+	T Value;
 };
 
-TEST_CASE("benchmarks.transaction_cost")
+template<typename T>
+class FGraph
 {
-#if defined(NDEBUG)
-	auto Sizes = {1, 100, 1000, 10000};
-#else
-	auto Sizes = {1, 100};
-#endif
+public:
+	FGraph() = default;
+
+	FGraph(const FGraph&) = delete;
+	FGraph& operator=(const FGraph&) = delete;
 
-	for (const unsigned Size : Sizes)
+	FGraph(FGraph&& InGraph)
 	{
-		std::vector<MyStruct> Datas(Size);
+		*this = std::move(InGraph);
+	}
+
+	FGraph& operator=(FGraph&& InGraph)
+	{
+		Reset();
+		Nodes = std::move(InGraph.Nodes);
+		Roots = std::move(InGraph.Roots);
+		return *this;
+	}
+
+	~FGraph()
+	{
+		Reset();
+	}
 
-		BENCHMARK("open " + std::to_string(Size))
+	void Reset()
+	{
+		for (FNode<T>* Node : Nodes)
 		{
-			for (unsigned i = 0; i < Size; i++)
+			delete Node;
+		}
+		Nodes.clear();
+		Roots.clear();
+	}
+
+	FNode<T>* AddNode(T Value)
+	{
+		FNode<T>* Node = new FNode<T>(Value);
+		Nodes.push_back(Node);
+		return Node;
+	}
+
+	void AddRoot(FNode<T>* Node)
+	{
+		Roots.push_back(Node);
+	}
+
+	const vector<FNode<T>*>& GetNodes() const { return Nodes; }
+	const vector<FNode<T>*>& GetRoots() const { return Roots; }
+
+	template<typename TFunc>
+	void DepthFirstSearchPre(const TFunc& Func) const
+	{
+		vector<FNode<T>*> Worklist;
+		unordered_set<FNode<T>*> Seen;
+
+		auto Push = [&](FNode<T>* Node)
 			{
-				Datas[i].DoSomething();
-			}
-		};
+				if (Seen.insert(Node).second)
+				{
+					Worklist.push_back(Node);
+				}
+			};
 
-		BENCHMARK("one transaction " + std::to_string(Size))
-		{
-			AutoRTFM::Transact([&] ()
+		auto PushAll = [&](const vector<FNode<T>*>& InnerNodes)
 			{
-				for (unsigned i = 0; i < Size; i++)
+				for (FNode<T>* Node : InnerNodes)
 				{
-					Datas[i].DoSomething();
+					Push(Node);
 				}
-			});
+			};
+
+		PushAll(Roots);
+
+		while (!Worklist.empty())
+		{
+			FNode<T>* Node = Worklist.back();
+			Worklist.pop_back();
+
+			Func(Node);
+
+			PushAll(Node->GetEdges());
+		}
+	}
+
+private:
+	vector<FNode<T>*> Nodes;
+	vector<FNode<T>*> Roots;
+};
+
+unsigned XorshiftState;
+
+void ResetXorshift()
+{
+	XorshiftState = 666;
+}
+
+unsigned Xorshift()
+{
+	unsigned Value = XorshiftState;
+	Value ^= Value << 13;
+	Value ^= Value >> 17;
+	Value ^= Value << 5;
+	XorshiftState = Value;
+	return Value;
+}
+
+unsigned BadRandom(unsigned Limit)
+{
+	return Xorshift() % Limit; // Yes I know this isn't great.
+}
+
+void AddToGraph(FGraph<unsigned>& Graph)
+{
+	auto RandomNode = [&](unsigned Value)
+		{
+			FNode<unsigned>* Result = Graph.AddNode(Value);
+			if (!BadRandom(Graph.GetNodes().size() + 1))
+			{
+				Graph.AddRoot(Result);
+			}
+			return Result;
 		};
 
-		BENCHMARK("many transactions " + std::to_string(Size))
+	FNode<unsigned>* A = RandomNode(Xorshift());
+	FNode<unsigned>* B = RandomNode(Xorshift());
+	FNode<unsigned>* C = RandomNode(Xorshift());
+	FNode<unsigned>* D = RandomNode(Xorshift());
+	FNode<unsigned>* E = RandomNode(Xorshift());
+	FNode<unsigned>* F = RandomNode(Xorshift());
+	FNode<unsigned>* G = RandomNode(Xorshift());
+
+	auto RandomEdge = [&](FNode<unsigned>* From, FNode<unsigned>* To)
 		{
-			for (unsigned i = 0; i < Size; i++)
+			if (!BadRandom(5))
 			{
-				AutoRTFM::Transact([&] ()
-				{
-					Datas[i].DoSomething();
-				});
+				From = Graph.GetNodes()[BadRandom(Graph.GetNodes().size())];
+			}
+			if (!BadRandom(5))
+			{
+				To = Graph.GetNodes()[BadRandom(Graph.GetNodes().size())];
 			}
+			From->AddEdge(To);
 		};
+
+	RandomEdge(A, B);
+	RandomEdge(B, C);
+	RandomEdge(C, B);
+	RandomEdge(C, D);
+	RandomEdge(B, D);
+	RandomEdge(D, E);
+	RandomEdge(E, F);
+	RandomEdge(E, G);
+	RandomEdge(A, E);
+	RandomEdge(A, G);
+	RandomEdge(G, E);
+}
+
+FGraph<unsigned> BuildGraph(const unsigned Total = 10000)
+{
+	FGraph<unsigned> Result;
+	for (unsigned Count = Total; Count--;)
+	{
+		AddToGraph(Result);
+	}
+	return Result;
+}
+
+unsigned WalkGraph(const FGraph<unsigned>& Graph)
+{
+	unsigned Result = 0;
+	unsigned Count = 0;
+	Graph.DepthFirstSearchPre([&](FNode<unsigned>* Node)
+		{
+			Result += Node->GetValue();
+			Count++;
+		});
+	//printf("Saw %zu roots, %u nodes, sum is %u.\n", Graph.GetRoots().size(), Count, Result);
+
+	return Result;
+}
+
+void CheckResult(unsigned Result, unsigned Total)
+{
+	switch (Total)
+	{
+	case 10000: REQUIRE(Result == 434344629); break;
+	case 100: REQUIRE(Result == 3732096243); break;
+	case 10: REQUIRE(Result == 3524276090); break;
+	case 1: REQUIRE(Result == 2218159753); break;
+	default: abort();
+	}
+}
+
+TEST_CASE("Graph")
+{
+	{
+		FGraph<unsigned> Graph;
+
+		AutoRTFM::Commit([&]() { Graph.AddNode(42); });
+
+		auto& Nodes = Graph.GetNodes();
+
+		REQUIRE(Nodes.size() == 1);
+		REQUIRE(Nodes[0]->GetValue() == 42);
+	}
+
+	{
+		FGraph<unsigned> Graph;
+
+		AutoRTFM::Commit([&]() { Graph.AddRoot(Graph.AddNode(42)); });
+
+		auto& Roots = Graph.GetRoots();
+
+		REQUIRE(Roots.size() == 1);
+		REQUIRE(Roots[0]->GetValue() == 42);
+
+		auto& Nodes = Graph.GetNodes();
+
+		REQUIRE(Nodes.size() == 1);
+		REQUIRE(Nodes[0]->GetValue() == 42);
 	}
+
+	for (unsigned Total : {1, 10, 100, 10000})
+	{
+		{
+			unsigned Result;
+			ResetXorshift();
+			FGraph<unsigned> Graph;
+			AutoRTFM::Commit([&]() { Graph = BuildGraph(Total); });
+			Result = WalkGraph(Graph);
+			CheckResult(Result, Total);
+		}
+
+		{
+			unsigned Result;
+			ResetXorshift();
+			FGraph<unsigned> Graph = BuildGraph(Total);
+			AutoRTFM::Commit([&]() { Result = WalkGraph(Graph); });
+			CheckResult(Result, Total);
+		}
+
+		{
+			unsigned Result;
+			ResetXorshift();
+			FGraph<unsigned> Graph;
+			AutoRTFM::Commit([&]() { Graph = BuildGraph(Total); });
+			AutoRTFM::Commit([&]() { Result = WalkGraph(Graph); });
+			CheckResult(Result, Total);
+		}
+
+		{
+			unsigned Result;
+			ResetXorshift();
+			FGraph<unsigned> Graph;
+			AutoRTFM::Commit([&]() { Graph = BuildGraph(Total); Result = WalkGraph(Graph); });
+			CheckResult(Result, Total);
+		}
+	}
+
+	BENCHMARK("build non transactional / walk non transactional")
+	{
+		ResetXorshift();
+		FGraph<unsigned> Graph = BuildGraph();
+		WalkGraph(Graph);
+	};
+
+	BENCHMARK("build transactional / walk non transactional")
+	{
+		ResetXorshift();
+		FGraph<unsigned> Graph;
+		AutoRTFM::Commit([&]() { Graph = BuildGraph(); });
+		WalkGraph(Graph);
+	};
+
+	BENCHMARK("build non transactional / walk transactional")
+	{
+		ResetXorshift();
+		FGraph<unsigned> Graph = BuildGraph();
+		AutoRTFM::Commit([&]() { WalkGraph(Graph); });
+	};
+
+	BENCHMARK("build transactional / walk transactional")
+	{
+		ResetXorshift();
+		FGraph<unsigned> Graph;
+		AutoRTFM::Commit([&]() { Graph = BuildGraph(); });
+		AutoRTFM::Commit([&]() { WalkGraph(Graph); });
+	};
+
+	BENCHMARK("build + walk transactional")
+	{
+		ResetXorshift();
+		FGraph<unsigned> Graph;
+		AutoRTFM::Commit([&]() { Graph = BuildGraph(); WalkGraph(Graph); });
+	};
 }
diff --git a/Engine/Source/Programs/AutoRTFMTests/Private/GraphBenchmark.cpp b/Engine/Source/Programs/AutoRTFMTests/Private/GraphBenchmark.cpp
deleted file mode 100644
diff --git a/Engine/Source/Runtime/Core/Public/AutoRTFM/AutoRTFM.h b/Engine/Source/Runtime/Core/Public/AutoRTFM/AutoRTFM.h
--- a/Engine/Source/Runtime/Core/Public/AutoRTFM/AutoRTFM.h
+++ b/Engine/Source/Runtime/Core/Public/AutoRTFM/AutoRTFM.h
@@ -143,9 +143,7 @@ UE_AUTORTFM_API void autortfm_commit(void (*work)(void* arg), void* arg);
 #else
 UE_AUTORTFM_FORCEINLINE void autortfm_commit(void (*work)(void* arg), void* arg)
 {
-	UE_AUTORTFM_UNUSED(work);
-	UE_AUTORTFM_UNUSED(arg);
-    abort();
+	work(arg);
 }
 #endif
 
