Fix warning during project file generation: "Overriding FPSemantics requires a private PCH. Disabling PCH usage for AutoRTFMTests"

diff --git a/Engine/Source/Programs/AutoRTFMTests/AutoRTFMTests.Build.cs b/Engine/Source/Programs/AutoRTFMTests/AutoRTFMTests.Build.cs
--- a/Engine/Source/Programs/AutoRTFMTests/AutoRTFMTests.Build.cs
+++ b/Engine/Source/Programs/AutoRTFMTests/AutoRTFMTests.Build.cs
@@ -33,6 +33,7 @@ public class AutoRTFMTests : ModuleRules
 				"Runtime/Core/Private"
 			});
 
+		PCHUsage = PCHUsageMode.NoPCHs;
 		FPSemantics = FPSemanticsMode.Precise;
 	}
 }
