For ReadOnly lock lets Open/PushAbort handler to unlock if we abort to handle the lock/unlock

diff --git a/Engine/Source/Programs/AutoRTFMTests/Private/RWLock.cpp b/Engine/Source/Programs/AutoRTFMTests/Private/RWLock.cpp
new file mode 100644
--- /dev/null
+++ b/Engine/Source/Programs/AutoRTFMTests/Private/RWLock.cpp
@@ -0,0 +1,43 @@
+// Copyright Epic Games, Inc. All Rights Reserved.
+
+#include "Catch2Includes.h"
+#include "Misc/ScopeRWLock.h"
+#include <AutoRTFM/AutoRTFM.h>
+#include <map>
+#include <vector>
+
+TEST_CASE("ReadLock")
+{
+	FRWLock ReadLock;
+	int x = 42;
+
+	auto Transaction = AutoRTFM::Transact([&]()
+    {
+		FReadScopeLock ScopeLock(ReadLock);
+		x = 43;
+	});
+
+    REQUIRE(x == 43);
+	// the read lock should have been released, we verify that by trying to acquire a write lock here
+    REQUIRE(ReadLock.TryWriteLock());
+}
+
+TEST_CASE("ReadLockAbort")
+{
+	FRWLock ReadLock;
+	int x = 42;
+
+	auto Transaction = AutoRTFM::Transact([&]()
+    {
+		FReadScopeLock ScopeLock(ReadLock);
+		x = 43;
+    	AutoRTFM::AbortTransaction();
+	});
+
+    REQUIRE(
+        AutoRTFM::ETransactionResult::AbortedByRequest ==
+		Transaction);
+    REQUIRE(x == 42);
+	// the read lock should have been released, we verify that by trying to acquire a write lock here
+    REQUIRE(ReadLock.TryWriteLock());
+}
diff --git a/Engine/Source/Runtime/Core/Public/Misc/ScopeRWLock.h b/Engine/Source/Runtime/Core/Public/Misc/ScopeRWLock.h
--- a/Engine/Source/Runtime/Core/Public/Misc/ScopeRWLock.h
+++ b/Engine/Source/Runtime/Core/Public/Misc/ScopeRWLock.h
@@ -5,6 +5,7 @@
 #include "CoreTypes.h"
 #include "Misc/AssertionMacros.h"
 #include "HAL/CriticalSection.h"
+#include "AutoRTFM/AutoRTFM.h"
 
 
 /** Keeps a FRWLock read-locked while this scope lives */
@@ -14,12 +15,14 @@ public:
 	UE_NODISCARD_CTOR explicit FReadScopeLock(FRWLock& InLock)
 		: Lock(InLock)
 	{
-		Lock.ReadLock();
+		UE_AUTORTFM_OPEN({ Lock.ReadLock(); });
+		AutoRTFM::PushOnAbortHandler(this, [this](){ this->Lock.ReadUnlock(); });
 	}
 
 	~FReadScopeLock()
 	{
-		Lock.ReadUnlock();
+		UE_AUTORTFM_OPEN({ Lock.ReadUnlock(); });
+		AutoRTFM::PopOnAbortHandler(this);
 	}
 
 private:
