Release the locks after trying to check if its write is not locked

diff --git a/Engine/Source/Programs/AutoRTFMTests/Private/RWLock.cpp b/Engine/Source/Programs/AutoRTFMTests/Private/RWLock.cpp
--- a/Engine/Source/Programs/AutoRTFMTests/Private/RWLock.cpp
+++ b/Engine/Source/Programs/AutoRTFMTests/Private/RWLock.cpp
@@ -20,6 +20,7 @@ TEST_CASE("ReadLock")
     REQUIRE(x == 43);
 	// the read lock should have been released, we verify that by trying to acquire a write lock here
     REQUIRE(ReadLock.TryWriteLock());
+	ReadLock.WriteUnlock();
 }
 
 TEST_CASE("ReadLockAbort")
@@ -40,4 +41,5 @@ TEST_CASE("ReadLockAbort")
     REQUIRE(x == 42);
 	// the read lock should have been released, we verify that by trying to acquire a write lock here
     REQUIRE(ReadLock.TryWriteLock());
+	ReadLock.WriteUnlock();
 }
