Fix AnnotationTests in AutoRTFMTests

diff --git a/Engine/Source/Programs/AutoRTFMTests/Private/UObjectTests.cpp b/Engine/Source/Programs/AutoRTFMTests/Private/UObjectTests.cpp
--- a/Engine/Source/Programs/AutoRTFMTests/Private/UObjectTests.cpp
+++ b/Engine/Source/Programs/AutoRTFMTests/Private/UObjectTests.cpp
@@ -142,12 +142,16 @@ namespace
 {
 	struct FTestAnnotation
 	{
-		// Default value, if it doesnt exist it will return a default value one
-		int TestAnnotationNumber = 32;
+		FTestAnnotation()
+			: TestAnnotationNumber(32)
+		{
+		}
+
+		int TestAnnotationNumber;
 
 		bool IsDefault() const
 		{
-			return false;
+			return TestAnnotationNumber == 32;
 		}
 	};
 	FUObjectAnnotationSparse<FTestAnnotation, true> GTestAnnotation;
