autortfm tests: fix ARM compile error

diff --git a/Engine/Source/Programs/AutoRTFMTests/Private/FunctionPointerTests.cpp b/Engine/Source/Programs/AutoRTFMTests/Private/FunctionPointerTests.cpp
--- a/Engine/Source/Programs/AutoRTFMTests/Private/FunctionPointerTests.cpp
+++ b/Engine/Source/Programs/AutoRTFMTests/Private/FunctionPointerTests.cpp
@@ -16,6 +16,8 @@ namespace
 		{
 			AutoRTFM::AbortTransaction();
 		}
+
+		return 43;
 	}
 
 	typedef int (*CStyleType)();
