Fix the UObject transactional annotation bug with the fix coming from Robert Manuszewski.

#rnx

diff --git a/Engine/Source/Programs/AutoRTFMTests/Private/UObjectTests.cpp b/Engine/Source/Programs/AutoRTFMTests/Private/UObjectTests.cpp
--- a/Engine/Source/Programs/AutoRTFMTests/Private/UObjectTests.cpp
+++ b/Engine/Source/Programs/AutoRTFMTests/Private/UObjectTests.cpp
@@ -1,7 +1,12 @@
 // Copyright Epic Games, Inc. All Rights Reserved.
 
-#include "MyAutoRTFMTestObject.h"
+#include "GenericPlatform/GenericPlatformMisc.h"
+
+THIRD_PARTY_INCLUDES_START
 #include "Catch2Includes.h"
+THIRD_PARTY_INCLUDES_END
+
+#include "MyAutoRTFMTestObject.h"
 #include <AutoRTFM/AutoRTFM.h>
 #include "UObject/GCObject.h"
 #include "UObject/ReachabilityAnalysis.h"
@@ -125,10 +130,10 @@ TEST_CASE("UObject.MarkAsReachable")
 			return true;
 		}
 
-		AutoRTFM::Commit([&]
+		REQUIRE(AutoRTFM::ETransactionResult::Committed == AutoRTFM::Transact([&]
 		{
 			Object->MarkAsReachable();
-		});
+		}));
 
 		return false;
 	});
@@ -202,3 +207,32 @@ TEST_CASE("UObject.TestAddAnnotation")
 		REQUIRE(32 == GTestAnnotation.GetAnnotation(Object).TestAnnotationNumber);
 	}
 }
+
+struct FAnnotationObject
+{
+	UObject* Object = nullptr;
+
+	FAnnotationObject() {}
+
+	FAnnotationObject(UObject* InObject) : Object(InObject) {}
+
+	bool IsDefault() { return !Object; }
+};
+
+template <> struct TIsPODType<FAnnotationObject> { enum { Value = true }; };
+
+TEST_CASE("UObject.AnnotationMap")
+{
+	FUObjectAnnotationSparse<FAnnotationObject, false> AnnotationMap;
+
+	UObject* Key = NewObject<UMyAutoRTFMTestObject>();
+
+	AutoRTFM::Transact([&]
+	{
+		UObject* Value = NewObject<UMyAutoRTFMTestObject>();
+		AnnotationMap.GetAnnotation(Key);
+		AnnotationMap.AddAnnotation(Key, Value);
+	});
+
+	REQUIRE(!AnnotationMap.GetAnnotation(Key).IsDefault());
+}
diff --git a/Engine/Source/Runtime/CoreUObject/Public/UObject/UObjectAnnotation.h b/Engine/Source/Runtime/CoreUObject/Public/UObject/UObjectAnnotation.h
--- a/Engine/Source/Runtime/CoreUObject/Public/UObject/UObjectAnnotation.h
+++ b/Engine/Source/Runtime/CoreUObject/Public/UObject/UObjectAnnotation.h
@@ -106,6 +106,8 @@ private:
 					else
 					{
 						AnnotationMap.Add(Object, LocalAnnotation);
+						AnnotationCacheKey = nullptr;
+						AnnotationCacheValue = TAnnotation();
 					}
 				};
 
