autortfm: handle masked store intrinsics
- have the compiler send the pointer and mask to the runtime so we can accurately capture which bytes are being written and log the undo

diff --git a/Engine/Source/Programs/AutoRTFMTests/AutoRTFMTests.Target.cs b/Engine/Source/Programs/AutoRTFMTests/AutoRTFMTests.Target.cs
--- a/Engine/Source/Programs/AutoRTFMTests/AutoRTFMTests.Target.cs
+++ b/Engine/Source/Programs/AutoRTFMTests/AutoRTFMTests.Target.cs
@@ -34,5 +34,7 @@ public class AutoRTFMTestsTarget : TargetRules
 		}
 
 		bFNameOutlineNumber = true;
+
+		MinCpuArchX64 = MinimumCpuArchitectureX64.AVX;
 	}
 }
diff --git a/Engine/Source/Programs/AutoRTFMTests/Private/Tests.cpp b/Engine/Source/Programs/AutoRTFMTests/Private/Tests.cpp
--- a/Engine/Source/Programs/AutoRTFMTests/Private/Tests.cpp
+++ b/Engine/Source/Programs/AutoRTFMTests/Private/Tests.cpp
@@ -9,6 +9,10 @@
 #include <cmath>
 #include "Logging/LogMacros.h"
 
+#if PLATFORM_CPU_X86_FAMILY
+#include <immintrin.h>
+#endif // PLATFORM_CPU_X86_FAMILY
+
 DECLARE_LOG_CATEGORY_EXTERN(LogAutoRTFMTests, Display, All)
 DEFINE_LOG_CATEGORY(LogAutoRTFMTests)
 
@@ -338,3 +342,108 @@ TEST_CASE("Tests.RetryNonNested")
 
 	REQUIRE(2 == Count);
 }
+
+#if PLATFORM_CPU_X86_FAMILY
+template<int WhichMaskOff>
+void AVXDoMaskedStore(double* Vector)
+{
+	__m256i Mask = _mm256_setr_epi64x(0xFFFFFFFFFFFFFFFFull, 0xFFFFFFFFFFFFFFFFull, 0xFFFFFFFFFFFFFFFFull, 0xFFFFFFFFFFFFFFFFull);
+	Mask = _mm256_insert_epi64(Mask, 0x0000000000000000ull, WhichMaskOff);
+	__m256d Val = _mm256_setr_pd(101.0, 102.0, 103.0, 104.0);
+	_mm256_maskstore_pd(Vector, Mask, Val);
+}
+
+TEST_CASE("Tests.AVXMaskedStore")
+{
+	auto RequireVectorsEqual = [](double* VectorPtr, __m256d Rhs) -> void
+	{
+		AutoRTFM::Open([&]()
+		{
+			__m256d Lhs = _mm256_load_pd(VectorPtr);
+			__m256 CmpResult = _mm256_castpd_ps(_mm256_cmp_pd(Lhs, Rhs, 0));
+			int Result = _mm256_movemask_ps(CmpResult);
+			REQUIRE(Result == 0xFF);
+		});
+	};
+
+	{
+		double Vector[4] = {1.0, 2.0, 3.0, 4.0};
+		auto TransactResult = AutoRTFM::Transact([&]()
+		{
+			// do a masked store to Vector
+			AVXDoMaskedStore<0>(&Vector[0]);
+
+			RequireVectorsEqual(&Vector[0], _mm256_setr_pd(1.0, 102.0, 103.0, 104.0));
+
+			// we overwrite the non-written value in the Open before we abort, to ensure
+			// that the runtime only rolls back elements that the masked write wrote to
+			AutoRTFM::Open([&]() { Vector[0] = 99.0; });
+
+    		AutoRTFM::AbortTransaction();
+		});
+
+		REQUIRE(TransactResult == AutoRTFM::ETransactionResult::AbortedByRequest);
+		RequireVectorsEqual(&Vector[0], _mm256_setr_pd(99.0, 2.0, 3.0, 4.0));
+	}
+
+	{
+		double Vector[4] = {1.0, 2.0, 3.0, 4.0};
+		auto TransactResult = AutoRTFM::Transact([&]()
+		{
+			// do a masked store to Vector
+			AVXDoMaskedStore<1>(&Vector[0]);
+
+			RequireVectorsEqual(&Vector[0], _mm256_setr_pd(101.0, 2.0, 103.0, 104.0));
+
+			// we overwrite the non-written value in the Open before we abort, to ensure
+			// that the runtime only rolls back elements that the masked write wrote to
+			AutoRTFM::Open([&]() { Vector[1] = 99.0; });
+
+    		AutoRTFM::AbortTransaction();
+		});
+
+		REQUIRE(TransactResult == AutoRTFM::ETransactionResult::AbortedByRequest);
+		RequireVectorsEqual(&Vector[0], _mm256_setr_pd(1.0, 99.0, 3.0, 4.0));
+	}
+
+	{
+		double Vector[4] = {1.0, 2.0, 3.0, 4.0};
+		auto TransactResult = AutoRTFM::Transact([&]()
+		{
+			// do a masked store to Vector
+			AVXDoMaskedStore<2>(&Vector[0]);
+
+			RequireVectorsEqual(&Vector[0], _mm256_setr_pd(101.0, 102.0, 3.0, 104.0));
+
+			// we overwrite the non-written value in the Open before we abort, to ensure
+			// that the runtime only rolls back elements that the masked write wrote to
+			AutoRTFM::Open([&]() { Vector[2] = 99.0; });
+
+    		AutoRTFM::AbortTransaction();
+		});
+
+		REQUIRE(TransactResult == AutoRTFM::ETransactionResult::AbortedByRequest);
+		RequireVectorsEqual(&Vector[0], _mm256_setr_pd(1.0, 2.0, 99.0, 4.0));
+	}
+
+	{
+		double Vector[4] = {1.0, 2.0, 3.0, 4.0};
+		auto TransactResult = AutoRTFM::Transact([&]()
+		{
+			// do a masked store to Vector
+			AVXDoMaskedStore<3>(&Vector[0]);
+
+			RequireVectorsEqual(&Vector[0], _mm256_setr_pd(101.0, 102.0, 103.0, 4.0));
+
+			// we overwrite the non-written value in the Open before we abort, to ensure
+			// that the runtime only rolls back elements that the masked write wrote to
+			AutoRTFM::Open([&]() { Vector[3] = 99.0; });
+
+    		AutoRTFM::AbortTransaction();
+		});
+
+		REQUIRE(TransactResult == AutoRTFM::ETransactionResult::AbortedByRequest);
+		RequireVectorsEqual(&Vector[0], _mm256_setr_pd(1.0, 2.0, 3.0, 99.0));
+	}
+}
+#endif // PLATFORM_CPU_X86_FAMILY
diff --git a/Engine/Source/Runtime/Core/Private/AutoRTFM/CompilerInterface.cpp b/Engine/Source/Runtime/Core/Private/AutoRTFM/CompilerInterface.cpp
--- a/Engine/Source/Runtime/Core/Private/AutoRTFM/CompilerInterface.cpp
+++ b/Engine/Source/Runtime/Core/Private/AutoRTFM/CompilerInterface.cpp
@@ -93,6 +93,29 @@ extern "C" UE_AUTORTFM_API AUTORTFM_NO_ASAN void autortfm_record_write_8(void* P
 	Context->RecordWrite<8>(Ptr);
 }
 
+extern "C" UE_AUTORTFM_API AUTORTFM_NO_ASAN void autortfm_record_masked_write(void* Ptr, uintptr_t Mask, int MaskWidthBits, int ValueSizeBytes)
+{
+	// check for writes to null here so we end up crashing in the user
+	// code rather than in the autortfm runtime.
+	if (UNLIKELY(nullptr == Ptr))
+	{
+		return;
+	}
+
+	FContext* Context = FContext::Get();
+
+	char* IncrementablePtr = static_cast<char*>(Ptr);
+	for(int i = 0; i < MaskWidthBits; i++)
+	{
+		if (Mask & (1u << i))
+		{
+			autortfm_record_write(IncrementablePtr, ValueSizeBytes);
+		}
+
+		IncrementablePtr += ValueSizeBytes;
+	}
+}
+
 extern "C" UE_AUTORTFM_API void* autortfm_lookup_function(void* OriginalFunction, const char* Where)
 {
 	FContext* Context = FContext::Get();
