Annotate RemoveCache() methods with UE_AUTORTFM_NOAUTORTFM, add more tests

RemoveCache() is not currently called from transactional code-paths, but if this changes, this could lead to state corruption. To prevent this, annotate the methods with UE_AUTORTFM_NOAUTORTFM.

Expand the tests to exercise FTextCache::FindOrCache() in various states, with transactions that abort and commit. Due to the lack of query APIs, there's no easy way to actually probe the state of the cache, so just check that FindOrCache() returns something sensible.

diff --git a/Engine/Source/Programs/AutoRTFMTests/Private/UECoreTests.cpp b/Engine/Source/Programs/AutoRTFMTests/Private/UECoreTests.cpp
--- a/Engine/Source/Programs/AutoRTFMTests/Private/UECoreTests.cpp
+++ b/Engine/Source/Programs/AutoRTFMTests/Private/UECoreTests.cpp
@@ -24,7 +24,7 @@ TEST_CASE("UECore.FDelegateHandle")
 		REQUIRE(AutoRTFM::ETransactionResult::AbortedByRequest == Result);
 		REQUIRE(!Handle.IsValid());
 	}
-	
+
 	REQUIRE(!Handle.IsValid());
 
 	SECTION("With Commit")
@@ -578,18 +578,223 @@ TEST_CASE("UECore.TIntrusiveReferenceController")
 	}
 }
 
-TEST_CASE("UECore.FTextCache")
+TEST_CASE("UECore.FText")
 {
 	FText Text;
 	REQUIRE(Text.IsEmpty());
 
-	AutoRTFM::ETransactionResult Result = AutoRTFM::Transact([&]
+	SECTION("With Abort")
+	{
+		AutoRTFM::ETransactionResult Result = AutoRTFM::Transact([&]
+			{
+				Text = FText::FromString(FString(TEXT("Sheesh")));
+				AutoRTFM::AbortTransaction();
+			});
+
+		REQUIRE(AutoRTFM::ETransactionResult::AbortedByRequest == Result);
+		REQUIRE(Text.IsEmpty());
+	}
+
+	SECTION("With Commit")
+	{
+		AutoRTFM::ETransactionResult Result = AutoRTFM::Transact([&]
+			{
+				Text = FText::FromString(FString(TEXT("Sheesh")));
+			});
+
+		REQUIRE(AutoRTFM::ETransactionResult::Committed == Result);
+		REQUIRE(!Text.IsEmpty());
+		REQUIRE(Text.ToString() == TEXT("Sheesh"));
+	}
+}
+
+TEST_CASE("UECore.FTextCache")
+{
+	// FTextCache is a singleton. Grab its reference.
+	FTextCache& Cache = FTextCache::Get();
+
+	// Use a fixed cache key for the tests below.
+	const FTextId Key{TEXT("NAMESPACE"), TEXT("KEY")};
+
+	// As FTextCache does not supply any way to query what's held in the cache,
+	// the best we can do here is to call FindOrCache() and check the returned
+	// FText strings are as expected.
+	auto CheckCacheHealthy = [&]
+	{
+		FText LookupA = Cache.FindOrCache(TEXT("VALUE"), Key);
+		REQUIRE(LookupA.ToString() == TEXT("VALUE"));
+		FText LookupB = Cache.FindOrCache(TEXT("REPLACEMENT"), Key);
+		REQUIRE(LookupB.ToString() == TEXT("REPLACEMENT"));
+		Cache.RemoveCache(Key);
+	};
+
+	SECTION("FindOrCache() Add new")
+	{
+		SECTION("With Abort")
 		{
-			FString S("Sheesh");
-			Text = FTextCache::Get().FindOrCache(ANSI_TO_TCHAR("Oh my"), FTextInspector::GetTextId(FText::FromString(S)));
-			AutoRTFM::AbortTransaction();
-		});
+			AutoRTFM::ETransactionResult Result = AutoRTFM::Transact([&]()
+				{
+					Cache.FindOrCache(TEXT("VALUE"), Key);
+					AutoRTFM::AbortTransaction();
+				});
 
-	REQUIRE(AutoRTFM::ETransactionResult::AbortedByRequest == Result);
-	REQUIRE(Text.IsEmpty());
+			REQUIRE(AutoRTFM::ETransactionResult::AbortedByRequest == Result);
+
+			CheckCacheHealthy();
+		}
+
+		SECTION("With Commit")
+		{
+			AutoRTFM::ETransactionResult Result = AutoRTFM::Transact([&]()
+				{
+					Cache.FindOrCache(TEXT("VALUE"), Key);
+				});
+
+			REQUIRE(AutoRTFM::ETransactionResult::Committed == Result);
+
+			CheckCacheHealthy();
+		}
+	}
+
+	SECTION("FindOrCache() Replace with same value")
+	{
+		SECTION("With Abort")
+		{
+			// Add an entry to the cache before the transaction
+			Cache.FindOrCache(TEXT("VALUE"), Key);
+
+			AutoRTFM::ETransactionResult Result = AutoRTFM::Transact([&]()
+				{
+					Cache.FindOrCache(TEXT("REPLACEMENT"), Key);
+					AutoRTFM::AbortTransaction();
+				});
+
+			REQUIRE(AutoRTFM::ETransactionResult::AbortedByRequest == Result);
+
+			CheckCacheHealthy();
+		}
+
+		SECTION("With Commit")
+		{
+			// Add an entry to the cache before the transaction
+			Cache.FindOrCache(TEXT("VALUE"), Key);
+
+			AutoRTFM::ETransactionResult Result = AutoRTFM::Transact([&]()
+				{
+					Cache.FindOrCache(TEXT("VALUE"), Key);
+				});
+
+			REQUIRE(AutoRTFM::ETransactionResult::Committed == Result);
+
+			CheckCacheHealthy();
+		}
+	}
+
+	SECTION("FindOrCache() Replace with different value")
+	{
+		SECTION("With Abort")
+		{
+			// Add an entry to the cache before the transaction
+			Cache.FindOrCache(TEXT("ORIGINAL"), Key);
+
+			AutoRTFM::ETransactionResult Result = AutoRTFM::Transact([&]()
+				{
+					Cache.FindOrCache(TEXT("REPLACEMENT"), Key);
+					AutoRTFM::AbortTransaction();
+				});
+
+			REQUIRE(AutoRTFM::ETransactionResult::AbortedByRequest == Result);
+
+			CheckCacheHealthy();
+		}
+
+		SECTION("With Commit")
+		{
+			// Add an entry to the cache before the transaction
+			Cache.FindOrCache(TEXT("ORIGINAL"), Key);
+
+			AutoRTFM::ETransactionResult Result = AutoRTFM::Transact([&]()
+				{
+					Cache.FindOrCache(TEXT("REPLACEMENT"), Key);
+				});
+
+			REQUIRE(AutoRTFM::ETransactionResult::Committed == Result);
+
+			CheckCacheHealthy();
+		}
+	}
+
+	static constexpr bool bSupportsTransactionalRemoveCache = false; // #jira SOL-6743
+	if (!bSupportsTransactionalRemoveCache)
+	{
+		return;
+	}
+
+	SECTION("RemoveCache()")
+	{
+		SECTION("With Abort")
+		{
+			// Add an entry to the cache before the transaction
+			Cache.FindOrCache(TEXT("VALUE"), Key);
+
+			AutoRTFM::ETransactionResult Result = AutoRTFM::Transact([&]()
+				{
+					Cache.RemoveCache(Key);
+					AutoRTFM::AbortTransaction();
+				});
+
+			REQUIRE(AutoRTFM::ETransactionResult::AbortedByRequest == Result);
+
+			CheckCacheHealthy();
+		}
+
+		SECTION("With Commit")
+		{
+			// Add an entry to the cache before the transaction
+			Cache.FindOrCache(TEXT("VALUE"), Key);
+
+			AutoRTFM::ETransactionResult Result = AutoRTFM::Transact([&]()
+				{
+					Cache.RemoveCache(Key);
+				});
+
+			REQUIRE(AutoRTFM::ETransactionResult::Committed == Result);
+
+			CheckCacheHealthy();
+		}
+	}
+
+
+	SECTION("Mixed Closed & Open")
+	{
+		SECTION("Closed: FindOrCache() Open: RemoveCache()")
+		{
+			SECTION("With Abort")
+			{
+				AutoRTFM::ETransactionResult Result = AutoRTFM::Transact([&]()
+					{
+						Cache.FindOrCache(TEXT("VALUE"), Key);
+						AutoRTFM::Open([&]{ Cache.RemoveCache(Key); });
+						AutoRTFM::AbortTransaction();
+					});
+
+				REQUIRE(AutoRTFM::ETransactionResult::AbortedByRequest == Result);
+
+				CheckCacheHealthy();
+			}
+
+			SECTION("With Commit")
+			{
+				AutoRTFM::ETransactionResult Result = AutoRTFM::Transact([&]()
+					{
+						Cache.FindOrCache(TEXT("VALUE"), Key);
+						AutoRTFM::Open([&]{ Cache.RemoveCache(Key); });
+					});
+
+				REQUIRE(AutoRTFM::ETransactionResult::Committed == Result);
+
+				CheckCacheHealthy();
+			}
+		}
+	}
 }
diff --git a/Engine/Source/Runtime/Core/Private/Internationalization/TextCache.cpp b/Engine/Source/Runtime/Core/Private/Internationalization/TextCache.cpp
--- a/Engine/Source/Runtime/Core/Private/Internationalization/TextCache.cpp
+++ b/Engine/Source/Runtime/Core/Private/Internationalization/TextCache.cpp
@@ -10,6 +10,14 @@
 #include "Misc/ScopeLock.h"
 #include "AutoRTFM/AutoRTFM.h"
 
+// Notes on the use of UE_AUTORTFM_ALWAYS_OPEN and UE_AUTORTFM_NOAUTORTFM:
+// It is currently unsafe to use the cache in the open while an uncommitted
+// transaction is in flight that has also touched the cache (#jira SOL-6743).
+// RemoveCache() is not currently reachable from a closed transaction, so these
+// are annotated with UE_AUTORTFM_NOAUTORTFM to prevent new transactional use.
+// FindOrCache() is currently used from transactional code paths, so is
+// annotated with UE_AUTORTFM_ALWAYS_OPEN.
+
 FTextCache& FTextCache::Get()
 {
 	return TLazySingleton<FTextCache>::Get();
@@ -63,12 +71,12 @@ UE_AUTORTFM_ALWAYS_OPEN FText FTextCache::FindOrCache(const TCHAR* InTextLiteral
 	return NewText;
 }
 
-void FTextCache::RemoveCache(const FTextId& InTextId)
+UE_AUTORTFM_NOAUTORTFM void FTextCache::RemoveCache(const FTextId& InTextId)
 {
 	return RemoveCache(MakeArrayView(&InTextId, 1));
 }
 
-void FTextCache::RemoveCache(TArrayView<const FTextId> InTextIds)
+UE_AUTORTFM_NOAUTORTFM void FTextCache::RemoveCache(TArrayView<const FTextId> InTextIds)
 {
 	FScopeLock Lock(&CachedTextCS);
 	for (const FTextId& TextId : InTextIds)
@@ -77,7 +85,7 @@ void FTextCache::RemoveCache(TArrayView<const FTextId> InTextIds)
 	}
 }
 
-void FTextCache::RemoveCache(const TSet<FTextId>& InTextIds)
+UE_AUTORTFM_NOAUTORTFM void FTextCache::RemoveCache(const TSet<FTextId>& InTextIds)
 {
 	FScopeLock Lock(&CachedTextCS);
 	for (const FTextId& TextId : InTextIds)
