[AutoRTFM] Expand CString tests

Test each of the functions in an aborted and committed transaction.

diff --git a/Engine/Source/Programs/AutoRTFMTests/Private/CString.cpp b/Engine/Source/Programs/AutoRTFMTests/Private/CString.cpp
--- a/Engine/Source/Programs/AutoRTFMTests/Private/CString.cpp
+++ b/Engine/Source/Programs/AutoRTFMTests/Private/CString.cpp
@@ -4,93 +4,212 @@
 #include <AutoRTFM/AutoRTFM.h>
 #include <cstring>
 
+#ifdef __clang__
+#define CLANG_BEGIN_DISABLE_OPTIMIZATIONS  _Pragma("clang optimize off")
+#define CLANG_END_DISABLE_OPTIMIZATIONS    _Pragma("clang optimize on")
+#else
+#define CLANG_BEGIN_DISABLE_OPTIMIZATIONS
+#define CLANG_END_DISABLE_OPTIMIZATIONS
+#endif
+
+#ifdef _MSC_VER
+#define MSVC_BEGIN_DISABLE_WARN_UNSAFE_FUNCTION \
+		__pragma(warning(push)) \
+		__pragma(warning(disable : 4996))
+#define MSVC_END_DISABLE_UNSAFE_FN  __pragma(warning(pop))
+#else
+#define MSVC_BEGIN_DISABLE_WARN_UNSAFE_FUNCTION
+#define MSVC_END_DISABLE_UNSAFE_FN
+#endif
+
 TEST_CASE("CString.memcpy")
 {
-	constexpr unsigned Size = 128;
-	const char* const From = "Hello, world!";
-	char To[Size];
-	memset(To, 0, Size);
-	AutoRTFM::Commit([&]
+	const char* const From = "Kittie says meow";
+	char To[] = "Doggie says woof";
+
+	SECTION("With Abort")
+	{
+		CLANG_BEGIN_DISABLE_OPTIMIZATIONS
+		AutoRTFM::ETransactionResult Result = AutoRTFM::Transact([&]
+		{
+			memcpy(To, From, 6);
+			AutoRTFM::AbortTransaction();
+		});
+		CLANG_END_DISABLE_OPTIMIZATIONS
+
+		REQUIRE(AutoRTFM::ETransactionResult::AbortedByRequest == Result);
+		REQUIRE("Doggie says woof" == std::string_view(To));
+	}
+
+	SECTION("With Commit")
 	{
-		memcpy(To, From, 5);
-	});
+		CLANG_BEGIN_DISABLE_OPTIMIZATIONS
+		AutoRTFM::Commit([&] { memcpy(To, From, 6); });
+		CLANG_END_DISABLE_OPTIMIZATIONS
 
-	REQUIRE(0 == strcmp("Hello", To));
+		REQUIRE("Kittie says woof" == std::string_view(To));
+	}
 }
 
 TEST_CASE("CString.memmove")
 {
-	constexpr unsigned Size = 128;
-	char To[Size] = "Hello, world!";
-	AutoRTFM::Commit([&]
+	char To[] = "Hello, world!";
+
+	SECTION("With Abort")
+	{
+		CLANG_BEGIN_DISABLE_OPTIMIZATIONS
+		AutoRTFM::ETransactionResult Result = AutoRTFM::Transact([&]
+		{
+			memmove(To + 7, To, 5);
+			AutoRTFM::AbortTransaction();
+		});
+		CLANG_END_DISABLE_OPTIMIZATIONS
+
+		REQUIRE(AutoRTFM::ETransactionResult::AbortedByRequest == Result);
+		REQUIRE("Hello, world!" == std::string_view(To));
+	}
+
+	SECTION("With Commit")
 	{
-		memmove(To + 7, To, 5);
-	});
+		CLANG_BEGIN_DISABLE_OPTIMIZATIONS
+		AutoRTFM::Commit([&] { memmove(To + 7, To, 5); });
+		CLANG_END_DISABLE_OPTIMIZATIONS
 
-	REQUIRE(0 == strcmp("Hello, Hello!", To));
+		REQUIRE("Hello, Hello!" == std::string_view(To));
+	}
 }
 
 TEST_CASE("CString.strcpy")
 {
-	constexpr unsigned Size = 128;
-	const char* const From = "Hello, world!";
-	char To[Size];
+	MSVC_BEGIN_DISABLE_WARN_UNSAFE_FUNCTION
+
+	const char From[] = "Kittie says meow";
+	char To[] = "Doggie says woof";
+	static_assert(sizeof(From) == sizeof(To));
 
-#pragma clang optimize off
-	AutoRTFM::Commit([&]
+	SECTION("With Abort")
+	{
+		CLANG_BEGIN_DISABLE_OPTIMIZATIONS
+		AutoRTFM::ETransactionResult Result = AutoRTFM::Transact([&]
 		{
 			strcpy(To, From);
+			AutoRTFM::AbortTransaction();
 		});
-#pragma clang optimize on
+		CLANG_END_DISABLE_OPTIMIZATIONS
 
-	REQUIRE(0 == strcmp(From, To));
+		REQUIRE(AutoRTFM::ETransactionResult::AbortedByRequest == Result);
+		REQUIRE("Doggie says woof" == std::string_view(To));
+	}
+
+	SECTION("With Commit")
+	{
+		CLANG_BEGIN_DISABLE_OPTIMIZATIONS
+		AutoRTFM::Commit([&] { strcpy(To, From); });
+		CLANG_END_DISABLE_OPTIMIZATIONS
+
+		REQUIRE("Kittie says meow" == std::string_view(To));
+	}
+
+	MSVC_END_DISABLE_UNSAFE_FN
 }
 
 TEST_CASE("CString.strncpy")
 {
-	constexpr unsigned Size = 128;
-	const char* const From = "Hello, world!";
-	char To[Size];
-	memset(To, 0, Size);
+	MSVC_BEGIN_DISABLE_WARN_UNSAFE_FUNCTION
+
+	const char* const From = "Kittie says meow";
+	char To[] = "Doggie says woof";
 
-#pragma clang optimize off
-	AutoRTFM::Commit([&]
+	SECTION("With Abort")
+	{
+		CLANG_BEGIN_DISABLE_OPTIMIZATIONS
+		AutoRTFM::ETransactionResult Result = AutoRTFM::Transact([&]
 		{
-			strncpy(To, From, 5);
+			strncpy(To, From, 6);
+			AutoRTFM::AbortTransaction();
 		});
-#pragma clang optimize on
+		CLANG_END_DISABLE_OPTIMIZATIONS
+
+		REQUIRE(AutoRTFM::ETransactionResult::AbortedByRequest == Result);
+		REQUIRE("Doggie says woof" == std::string_view(To));
+	}
+
+	SECTION("With Commit")
+	{
+		CLANG_BEGIN_DISABLE_OPTIMIZATIONS
+		AutoRTFM::Commit([&] { strncpy(To, From, 6); });
+		CLANG_END_DISABLE_OPTIMIZATIONS
+
+		REQUIRE("Kittie says woof" == std::string_view(To));
+	}
 
-	REQUIRE(0 == strcmp("Hello", To));
+	MSVC_END_DISABLE_UNSAFE_FN
 }
 
 TEST_CASE("CString.strcat")
 {
+	MSVC_BEGIN_DISABLE_WARN_UNSAFE_FUNCTION
+
 	constexpr unsigned Size = 128;
 	char To[Size] = "Hello";
 
-#pragma clang optimize off
-	AutoRTFM::Commit([&]
+	SECTION("With Abort")
+	{
+		CLANG_BEGIN_DISABLE_OPTIMIZATIONS
+		AutoRTFM::ETransactionResult Result = AutoRTFM::Transact([&]
 		{
 			strcat(To, ", world!");
+			AutoRTFM::AbortTransaction();
 		});
-#pragma clang optimize on
+		CLANG_END_DISABLE_OPTIMIZATIONS
+
+		REQUIRE(AutoRTFM::ETransactionResult::AbortedByRequest == Result);
+		REQUIRE("Hello" == std::string_view(To));
+	}
+
+	SECTION("With Commit")
+	{
+		CLANG_BEGIN_DISABLE_OPTIMIZATIONS
+		AutoRTFM::Commit([&] { strcat(To, ", world!"); });
+		CLANG_END_DISABLE_OPTIMIZATIONS
+
+		REQUIRE("Hello, world!" == std::string_view(To));
+	}
 
-	REQUIRE(0 == strcmp("Hello, world!", To));
+	MSVC_END_DISABLE_UNSAFE_FN
 }
 
 TEST_CASE("CString.strncat")
 {
+	MSVC_BEGIN_DISABLE_WARN_UNSAFE_FUNCTION
+
 	constexpr unsigned Size = 128;
 	char To[Size] = "Hello";
 
-#pragma clang optimize off
-	AutoRTFM::Commit([&]
+	SECTION("With Abort")
+	{
+		CLANG_BEGIN_DISABLE_OPTIMIZATIONS
+		AutoRTFM::ETransactionResult Result = AutoRTFM::Transact([&]
 		{
 			strncat(To, ", world! Not this!", 8);
+			AutoRTFM::AbortTransaction();
 		});
-#pragma clang optimize on
+		CLANG_END_DISABLE_OPTIMIZATIONS
+
+		REQUIRE(AutoRTFM::ETransactionResult::AbortedByRequest == Result);
+		REQUIRE("Hello" == std::string_view(To));
+	}
+
+	SECTION("With Commit")
+	{
+		CLANG_BEGIN_DISABLE_OPTIMIZATIONS
+		AutoRTFM::Commit([&] { strncat(To, ", world! Not this!", 8); });
+		CLANG_END_DISABLE_OPTIMIZATIONS
+
+		REQUIRE("Hello, world!" == std::string_view(To));
+	}
 
-	REQUIRE(0 == strcmp("Hello, world!", To));
+	MSVC_END_DISABLE_UNSAFE_FN
 }
 
 TEST_CASE("CString.memcmp")
@@ -98,109 +217,220 @@ TEST_CASE("CString.memcmp")
 	constexpr unsigned Size = 128;
 	char A[Size] = "This";
 
-	int Compare;
-#pragma clang optimize off
-	AutoRTFM::Commit([&]
+	SECTION("With Abort")
+	{
+		int Compare = 0;
+
+		CLANG_BEGIN_DISABLE_OPTIMIZATIONS
+		AutoRTFM::ETransactionResult Result = AutoRTFM::Transact([&]
 		{
 			Compare = memcmp(A, "That", 4);
+			AutoRTFM::AbortTransaction();
 		});
-#pragma clang optimize on
+		CLANG_END_DISABLE_OPTIMIZATIONS
+
+		REQUIRE(0 == Compare);
+	}
+
+	SECTION("With Commit")
+	{
+		int Compare = 0;
+
+		CLANG_BEGIN_DISABLE_OPTIMIZATIONS
+		AutoRTFM::Commit([&] { Compare = memcmp(A, "That", 4); });
+		CLANG_END_DISABLE_OPTIMIZATIONS
 
-	REQUIRE(0 < Compare);
+		REQUIRE(0 < Compare);
+	}
 }
 
 TEST_CASE("CString.strcmp")
 {
-	constexpr unsigned Size = 128;
-	char A[Size] = "This";
+	const char* A = "This";
+
+	SECTION("With Abort")
+	{
+		int Compare = 0;
 
-	int Compare;
-#pragma clang optimize off
-	AutoRTFM::Commit([&]
+		CLANG_BEGIN_DISABLE_OPTIMIZATIONS
+		AutoRTFM::ETransactionResult Result = AutoRTFM::Transact([&]
 		{
 			Compare = strcmp(A, "That");
+			AutoRTFM::AbortTransaction();
 		});
-#pragma clang optimize on
+		CLANG_END_DISABLE_OPTIMIZATIONS
+
+		REQUIRE(0 == Compare);
+	}
+
+	SECTION("With Commit")
+	{
+		int Compare = 0;
 
-	REQUIRE(0 < Compare);
+		CLANG_BEGIN_DISABLE_OPTIMIZATIONS
+		AutoRTFM::Commit([&] { Compare = strcmp(A, "That"); });
+		CLANG_END_DISABLE_OPTIMIZATIONS
+
+		REQUIRE(0 < Compare);
+	}
 }
 
 TEST_CASE("CString.strncmp")
 {
-	constexpr unsigned Size = 128;
-	char A[Size] = "This";
+	const char* A = "This";
 
-	int Compare;
-#pragma clang optimize off
-	AutoRTFM::Commit([&]
+	SECTION("With Abort")
+	{
+		int Compare = 0;
+
+		CLANG_BEGIN_DISABLE_OPTIMIZATIONS
+		AutoRTFM::ETransactionResult Result = AutoRTFM::Transact([&]
 		{
 			Compare = strncmp(A, "That", 3);
+			AutoRTFM::AbortTransaction();
 		});
-#pragma clang optimize on
+		CLANG_END_DISABLE_OPTIMIZATIONS
+
+		REQUIRE(0 == Compare);
+	}
 
-	REQUIRE(0 < Compare);
+	SECTION("With Commit")
+	{
+		int Compare = 0;
+
+		CLANG_BEGIN_DISABLE_OPTIMIZATIONS
+		AutoRTFM::Commit([&] { Compare = strncmp(A, "That", 3); });
+		CLANG_END_DISABLE_OPTIMIZATIONS
+
+		REQUIRE(0 < Compare);
+	}
 }
 
 TEST_CASE("CString.strchr")
 {
-	constexpr unsigned Size = 128;
-	char A[Size] = "This";
+	const char* A = "Thinking";
+
+	SECTION("With Abort")
+	{
+		const char* Value = nullptr;
 
-	const char* Result;
-#pragma clang optimize off
-	AutoRTFM::Commit([&]
+		CLANG_BEGIN_DISABLE_OPTIMIZATIONS
+		AutoRTFM::ETransactionResult Result = AutoRTFM::Transact([&]
 		{
-			Result = strchr(A, 'i');
+			Value = strchr(A, 'i');
+			AutoRTFM::AbortTransaction();
 		});
-#pragma clang optimize on
+		CLANG_END_DISABLE_OPTIMIZATIONS
+
+		REQUIRE(nullptr == Value);
+	}
+
+	SECTION("With Commit")
+	{
+		const char* Value = nullptr;
+
+		CLANG_BEGIN_DISABLE_OPTIMIZATIONS
+		AutoRTFM::Commit([&] { Value = strchr(A, 'i'); });
+		CLANG_END_DISABLE_OPTIMIZATIONS
 
-	REQUIRE((A + 2) == Result);
+		REQUIRE((A + 2) == Value);
+	}
 }
 
 TEST_CASE("CString.strrchr")
 {
-	constexpr unsigned Size = 128;
-	char A[Size] = "This";
+	const char* A = "Thinking";
+
+	SECTION("With Abort")
+	{
+		const char* Value = nullptr;
 
-	const char* Result;
-#pragma clang optimize off
-	AutoRTFM::Commit([&]
+		CLANG_BEGIN_DISABLE_OPTIMIZATIONS
+		AutoRTFM::ETransactionResult Result = AutoRTFM::Transact([&]
 		{
-			Result = strrchr(A, 'i');
+			Value = strrchr(A, 'i');
+			AutoRTFM::AbortTransaction();
 		});
-#pragma clang optimize on
+		CLANG_END_DISABLE_OPTIMIZATIONS
+
+		REQUIRE(nullptr == Value);
+	}
+
+	SECTION("With Commit")
+	{
+		const char* Value = nullptr;
 
-	REQUIRE((A + 2) == Result);
+		CLANG_BEGIN_DISABLE_OPTIMIZATIONS
+		AutoRTFM::Commit([&] { Value = strrchr(A, 'i'); });
+		CLANG_END_DISABLE_OPTIMIZATIONS
+
+		REQUIRE((A + 5) == Value);
+	}
 }
 
 TEST_CASE("CString.strstr")
 {
-	constexpr unsigned Size = 128;
-	char A[Size] = "This";
+	const char* A = "This";
 
-	const char* Result;
-#pragma clang optimize off
-	AutoRTFM::Commit([&]
+	SECTION("With Abort")
+	{
+		const char* Value = nullptr;
+
+		CLANG_BEGIN_DISABLE_OPTIMIZATIONS
+		AutoRTFM::ETransactionResult Result = AutoRTFM::Transact([&]
 		{
-			Result = strstr(A, "is");
+			Value = strstr(A, "is");
+			AutoRTFM::AbortTransaction();
 		});
-#pragma clang optimize on
+		CLANG_END_DISABLE_OPTIMIZATIONS
+
+		REQUIRE(nullptr == Value);
+	}
 
-	REQUIRE((A + 2) == Result);
+	SECTION("With Commit")
+	{
+		const char* Value = nullptr;
+
+		CLANG_BEGIN_DISABLE_OPTIMIZATIONS
+		AutoRTFM::Commit([&] { Value = strstr(A, "is"); });
+		CLANG_END_DISABLE_OPTIMIZATIONS
+
+		REQUIRE((A + 2) == Value);
+	}
 }
 
 TEST_CASE("CString.strlen")
 {
-	constexpr unsigned Size = 128;
-	char A[Size] = "This";
+	const char* A = "This";
+
+	SECTION("With Abort")
+	{
+		size_t Value = 0;
 
-	size_t Result;
-#pragma clang optimize off
-	AutoRTFM::Commit([&]
+		CLANG_BEGIN_DISABLE_OPTIMIZATIONS
+		AutoRTFM::ETransactionResult Result = AutoRTFM::Transact([&]
 		{
-			Result = strlen(A);
+			Value = strlen(A);
+			AutoRTFM::AbortTransaction();
 		});
-#pragma clang optimize on
+		CLANG_END_DISABLE_OPTIMIZATIONS
+
+		REQUIRE(0 == Value);
+	}
+
+	SECTION("With Commit")
+	{
+		size_t Value = 0;
+
+		CLANG_BEGIN_DISABLE_OPTIMIZATIONS
+		AutoRTFM::Commit([&] { Value = strlen(A); });
+		CLANG_END_DISABLE_OPTIMIZATIONS
 
-	REQUIRE(4 == Result);
+		REQUIRE(4 == Value);
+	}
 }
+
+#undef CLANG_BEGIN_DISABLE_OPTIMIZATIONS
+#undef CLANG_END_DISABLE_OPTIMIZATIONS
+#undef MSVC_BEGIN_DISABLE_WARN_UNSAFE_FUNCTION
+#undef MSVC_END_DISABLE_UNSAFE_FN
