[AutoRTFM] Add tests for math functions. Add missing self-registered functions

Add UE_AUTORTFM_REGISTER_SELF_FUNCTION() declarations for:
* asinhf
* acoshf
* atanhf
* atan2f
* llrintf

These are transitively called by their non-f suffixed, float variants.

diff --git a/Engine/Source/Programs/AutoRTFMTests/Private/Math.cpp b/Engine/Source/Programs/AutoRTFMTests/Private/Math.cpp
new file mode 100644
--- /dev/null
+++ b/Engine/Source/Programs/AutoRTFMTests/Private/Math.cpp
@@ -0,0 +1,280 @@
+// Copyright Epic Games, Inc. All Rights Reserved.
+
+#include "Catch2Includes.h"
+#include <AutoRTFM/AutoRTFM.h>
+#include <math.h>
+
+namespace
+{
+
+// Check calls Function outside a transaction, in an aborted transaction and
+// committed transaction. The value returned when called in the committed
+// transation is expected to match the value returned when called outside the
+// transaction.
+// Function must be the deterministic function with the signature 'T()'.
+template<typename FUNC>
+void Check(FUNC&& Function)
+{
+	using T = decltype(Function());
+	const T Expected = Function();
+	const T Zero{};
+
+	SECTION("With Abort")
+	{
+		T Got = Zero;
+		AutoRTFM::ETransactionResult Result = AutoRTFM::Transact([&]
+			{
+				Got = Function();
+				AutoRTFM::AbortTransaction();
+			});
+
+		REQUIRE(AutoRTFM::ETransactionResult::AbortedByRequest == Result);
+		REQUIRE(Zero == Got);
+	}
+
+	SECTION("With Commit")
+	{
+		T Got = Zero;
+
+		AutoRTFM::Commit([&] { Got = Function(); });
+
+		REQUIRE(Expected == Got);
+	}
+}
+
+}  // anonymous namespace
+
+TEST_CASE("Math.Sqrt")
+{
+	SECTION("float")  { Check([]{ return sqrt(0.42f); }); };
+	SECTION("double") { Check([]{ return sqrt(0.42); }); };
+}
+
+TEST_CASE("Math.Sqrtf")
+{
+	Check([]{ return sqrtf(0.42f); });
+}
+
+TEST_CASE("Math.Sin")
+{
+	SECTION("float")  { Check([]{ return sin(0.42f); }); }
+	SECTION("double") { Check([]{ return sin(0.42); }); }
+}
+
+TEST_CASE("Math.Sinf")
+{
+	Check([]{ return sinf(0.42f); });
+}
+
+TEST_CASE("Math.Cos")
+{
+	SECTION("float")  { Check([]{ return cos(0.42f); }); }
+	SECTION("double") { Check([]{ return cos(0.42); }); }
+}
+
+TEST_CASE("Math.Cosf")
+{
+	Check([]{ return cosf(0.42f); });
+}
+
+TEST_CASE("Math.Tan")
+{
+	SECTION("float")  { Check([]{ return tan(0.42f); }); }
+	SECTION("double") { Check([]{ return tan(0.42); }); }
+}
+
+TEST_CASE("Math.Tanf")
+{
+	Check([]{ return tanf(0.42f); });
+}
+
+TEST_CASE("Math.Asin")
+{
+	SECTION("float")  { Check([]{ return asin(0.42f); }); }
+	SECTION("double") { Check([]{ return asin(0.42); }); }
+}
+
+TEST_CASE("Math.Asinf")
+{
+	Check([]{ return asinf(0.42f); });
+}
+
+TEST_CASE("Math.Acos")
+{
+	SECTION("float")  { Check([]{ return acos(0.42f); }); }
+	SECTION("double") { Check([]{ return acos(0.42); }); }
+}
+
+TEST_CASE("Math.Acosf")
+{
+	Check([]{ return acosf(0.42f); });
+}
+
+TEST_CASE("Math.Atan")
+{
+	SECTION("float")  { Check([]{ return atan(0.42f); }); }
+	SECTION("double") { Check([]{ return atan(0.42); }); }
+}
+
+TEST_CASE("Math.Atanf")
+{
+	Check([]{ return atanf(0.42f); });
+}
+
+TEST_CASE("Math.Atan2")
+{
+	SECTION("float")  { Check([]{ return atan2(0.42f, 0.42f); }); }
+	SECTION("double") { Check([]{ return atan2(0.42, 0.42); }); }
+}
+
+TEST_CASE("Math.Atan2f")
+{
+	Check([]{ return atan2f(0.42f, 0.24f); });
+}
+
+TEST_CASE("Math.Sinh")
+{
+	SECTION("float")  { Check([]{ return sinh(0.42f); }); }
+	SECTION("double") { Check([]{ return sinh(0.42); }); }
+}
+
+TEST_CASE("Math.Sinhf")
+{
+	Check([]{ return sinhf(0.42f); });
+}
+
+TEST_CASE("Math.Cosh")
+{
+	SECTION("float")  { Check([]{ return cosh(0.42f); }); }
+	SECTION("double") { Check([]{ return cosh(0.42); }); }
+}
+
+TEST_CASE("Math.Coshf")
+{
+	Check([]{ return coshf(0.42f); });
+}
+
+TEST_CASE("Math.Tanh")
+{
+	SECTION("float")  { Check([]{ return tanh(0.42f); }); }
+	SECTION("double") { Check([]{ return tanh(0.42); }); }
+}
+
+TEST_CASE("Math.Tanhf")
+{
+	Check([]{ return tanhf(0.42f); });
+}
+
+TEST_CASE("Math.Asinh")
+{
+	SECTION("float")  { Check([]{ return asinh(0.42f); }); }
+	SECTION("double") { Check([]{ return asinh(0.42); }); }
+}
+
+TEST_CASE("Math.Asinhf")
+{
+	Check([]{ return asinhf(0.42f); });
+}
+
+TEST_CASE("Math.Acosh")
+{
+	SECTION("float")  { Check([]{ return acosh(4.2f); }); }
+	SECTION("double") { Check([]{ return acosh(4.2); }); }
+}
+
+TEST_CASE("Math.Acoshf")
+{
+	Check([]{ return acoshf(4.2f); });
+}
+
+TEST_CASE("Math.Atanh")
+{
+	SECTION("float")  { Check([]{ return atanh(0.42f); }); }
+	SECTION("double") { Check([]{ return atanh(0.42); }); }
+}
+
+TEST_CASE("Math.Atanhf")
+{
+	Check([]{ return atanhf(0.42f); });
+}
+
+TEST_CASE("Math.Exp")
+{
+	SECTION("float")  { Check([]{ return exp(0.42f); }); }
+	SECTION("double") { Check([]{ return exp(0.42); }); }
+}
+
+TEST_CASE("Math.Expf")
+{
+	Check([]{ return expf(0.42f); });
+}
+
+TEST_CASE("Math.Log")
+{
+	SECTION("float")  { Check([]{ return log(0.42f); }); }
+	SECTION("double") { Check([]{ return log(0.42); }); }
+}
+
+TEST_CASE("Math.Pow")
+{
+	SECTION("float")  { Check([]{ return pow(0.42f, 0.42f); }); }
+	SECTION("double") { Check([]{ return pow(0.42, 0.42); }); }
+}
+
+TEST_CASE("Math.Powf")
+{
+	Check([]{ return powf(0.42f, 0.24f); });
+}
+
+TEST_CASE("Math.Logf")
+{
+	Check([]{ return logf(0.42f); });
+}
+
+TEST_CASE("Math.Llrint")
+{
+	SECTION("float")  { Check([]{ return llrint(0.42f); }); }
+	SECTION("double") { Check([]{ return llrint(0.42); }); }
+}
+
+TEST_CASE("Math.Llrintf")
+{
+	Check([]{ return llrintf(0.42f); });
+}
+
+TEST_CASE("Math.Fmod")
+{
+	SECTION("float")  { Check([]{ return fmod(0.42f, 0.42f); }); }
+	SECTION("double") { Check([]{ return fmod(0.42, 0.42); }); }
+}
+
+TEST_CASE("Math.Fmodf")
+{
+	Check([]{ return fmodf(0.42f, 0.24f); });
+}
+
+TEST_CASE("Math.Fmodl")
+{
+	Check([]{ return fmodl(0.42, 0.24); });
+}
+
+TEST_CASE("Math.Rand")
+{
+	// it's near impossible to test the value returned by rand(), so just check
+	// that calling it a transaction doesn't explode.
+
+	SECTION("With Abort")
+	{
+		AutoRTFM::ETransactionResult Result = AutoRTFM::Transact([&]
+			{
+				rand();
+				AutoRTFM::AbortTransaction();
+			});
+		REQUIRE(AutoRTFM::ETransactionResult::AbortedByRequest == Result);
+	}
+
+	SECTION("With Commit")
+	{
+		AutoRTFM::Commit([&] { rand(); });
+	}
+}
diff --git a/Engine/Source/Runtime/Core/Private/AutoRTFM/StdLibSupport.cpp b/Engine/Source/Runtime/Core/Private/AutoRTFM/StdLibSupport.cpp
--- a/Engine/Source/Runtime/Core/Private/AutoRTFM/StdLibSupport.cpp
+++ b/Engine/Source/Runtime/Core/Private/AutoRTFM/StdLibSupport.cpp
@@ -208,14 +208,19 @@ UE_AUTORTFM_REGISTER_SELF_FUNCTION(sinf);
 UE_AUTORTFM_REGISTER_SELF_FUNCTION(cosf);
 UE_AUTORTFM_REGISTER_SELF_FUNCTION(tanf);
 UE_AUTORTFM_REGISTER_SELF_FUNCTION(asinf);
+UE_AUTORTFM_REGISTER_SELF_FUNCTION(asinhf);
 UE_AUTORTFM_REGISTER_SELF_FUNCTION(acosf);
+UE_AUTORTFM_REGISTER_SELF_FUNCTION(acoshf);
 UE_AUTORTFM_REGISTER_SELF_FUNCTION(atanf);
+UE_AUTORTFM_REGISTER_SELF_FUNCTION(atanhf);
+UE_AUTORTFM_REGISTER_SELF_FUNCTION(atan2f);
 UE_AUTORTFM_REGISTER_SELF_FUNCTION(sinhf);
 UE_AUTORTFM_REGISTER_SELF_FUNCTION(coshf);
 UE_AUTORTFM_REGISTER_SELF_FUNCTION(tanhf);
 UE_AUTORTFM_REGISTER_SELF_FUNCTION(expf);
 UE_AUTORTFM_REGISTER_SELF_FUNCTION(logf);
 UE_AUTORTFM_REGISTER_SELF_FUNCTION(powf);
+UE_AUTORTFM_REGISTER_SELF_FUNCTION(llrintf);
 UE_AUTORTFM_REGISTER_SELF_FUNCTION(fmodf);
 UE_AUTORTFM_REGISTER_SELF_FUNCTION(fmodl);
 UE_AUTORTFM_REGISTER_SELF_FUNCTION(rand);
