Disable writing out the logfile when running `AutoRTFMTests` by default to avoid spurious CI failures.

#rnx

diff --git a/Engine/Source/Programs/AutoRTFMTests/AutoRTFMTests.Target.cs b/Engine/Source/Programs/AutoRTFMTests/AutoRTFMTests.Target.cs
--- a/Engine/Source/Programs/AutoRTFMTests/AutoRTFMTests.Target.cs
+++ b/Engine/Source/Programs/AutoRTFMTests/AutoRTFMTests.Target.cs
@@ -1,11 +1,16 @@
 // Copyright Epic Games, Inc. All Rights Reserved.
 
+using EpicGames.Core;
 using UnrealBuildTool;
 using System.Collections.Generic;
 
 [SupportedPlatforms(UnrealPlatformClass.All)]
 public class AutoRTFMTestsTarget : TargetRules
 {
+	// TODO: Might be useful to promote this to a general Target.cs setting at some point in the future.
+	[CommandLine("-AllowLogFile")]
+	public bool bAllowLogFile = false;
+
 	public AutoRTFMTestsTarget(TargetInfo Target) : base(Target)
 	{
 		Type = TargetType.Program;
@@ -36,5 +41,16 @@ public class AutoRTFMTestsTarget : TargetRules
 		bFNameOutlineNumber = true;
 
 		MinCpuArchX64 = MinimumCpuArchitectureX64.AVX;
+
+		// Allow for disabling writing out the logfile, since in `PreSubmitTest.py` we run this target simultaneously
+		// multiple times, and doing so would cause writing them out to stomp each other.
+		if (!bAllowLogFile)
+		{
+			GlobalDefinitions.Add("ALLOW_LOG_FILE=0");
+		}
+		else
+		{
+			GlobalDefinitions.Add("ALLOW_LOG_FILE=1");
+		}
 	}
 }
