[AutoRTFM] Flesh out allocator tests, fix realloc, implement calloc

Expand the existing allocator AutoRTFM tests to test both a transaction commit and abort.
Enable the ArrayDelete test, as this was already implemented.
Implement MemcpyToNew(), used by AutoRTFM realloc.
Implement RTFM_calloc().
Add new tests for: free(), malloc(), calloc() and realloc()

diff --git a/Engine/Source/Programs/AutoRTFMTests/Private/Allocators.cpp b/Engine/Source/Programs/AutoRTFMTests/Private/Allocators.cpp
--- a/Engine/Source/Programs/AutoRTFMTests/Private/Allocators.cpp
+++ b/Engine/Source/Programs/AutoRTFMTests/Private/Allocators.cpp
@@ -3,79 +3,369 @@
 #include "Catch2Includes.h"
 #include <AutoRTFM/AutoRTFM.h>
 #include <new>
+#include <stdlib.h>
+
+#ifdef __clang__
+#define CLANG_BEGIN_DISABLE_OPTIMIZATIONS  _Pragma("clang optimize off")
+#define CLANG_END_DISABLE_OPTIMIZATIONS    _Pragma("clang optimize on")
+#else
+#define CLANG_BEGIN_DISABLE_OPTIMIZATIONS
+#define CLANG_END_DISABLE_OPTIMIZATIONS
+#endif
+
+#ifdef _MSC_VER
+#define MSVC_BEGIN_DISABLE_WARN_ALLOC_MAYBE_NULL \
+		__pragma(warning(push)) \
+		__pragma(warning(disable : 6011))
+#define MSVC_END_DISABLE_WARN_ALLOC_MAYBE_NULL  __pragma(warning(pop))
+#else
+#define MSVC_BEGIN_DISABLE_WARN_ALLOC_MAYBE_NULL
+#define MSVC_END_DISABLE_WARN_ALLOC_MAYBE_NULL
+#endif
 
 TEST_CASE("Allocators.New")
 {
-	int* Data;
-	AutoRTFM::Commit([&]
-		{
-			Data = new int(42);
-		});
+	SECTION("With Abort")
+	{
+		int* Data = nullptr;
+		AutoRTFM::ETransactionResult Result = AutoRTFM::Transact([&]
+			{
+				Data = new int(42);
+				AutoRTFM::AbortTransaction();
+			});
+
+		REQUIRE(AutoRTFM::ETransactionResult::AbortedByRequest == Result);
+		REQUIRE(nullptr == Data);
+	}
 
-	REQUIRE(42 == *Data);
+	SECTION("With Commit")
+	{
+		int* Data = nullptr;
+		AutoRTFM::Commit([&]
+			{
+				Data = new int(42);
+			});
 
-	delete Data;
+		REQUIRE(42 == *Data);
+
+		delete Data;
+	}
 }
 
 TEST_CASE("Allocators.Delete")
 {
-	int* Data = new int(42);
+	SECTION("With Abort")
+	{
+		int* Data = new int(42);
+
+		AutoRTFM::ETransactionResult Result = AutoRTFM::Transact([&]
+			{
+				delete Data;
+				AutoRTFM::AbortTransaction();
+			});
 
-	AutoRTFM::Commit([&]
-		{
-			delete Data;
-		});
+		REQUIRE(AutoRTFM::ETransactionResult::AbortedByRequest == Result);
+		REQUIRE(42 == *Data);
+
+		delete Data;
+	}
+
+	SECTION("With Commit")
+	{
+		int* Data = new int(42);
+
+		AutoRTFM::Commit([&]
+			{
+				delete Data;
+			});
+	}
 }
 
 TEST_CASE("Allocators.ArrayNew")
 {
-	int* Data;
-	AutoRTFM::Commit([&]
-		{
-			Data = new int[42]();
-			Data[2] = 42;
-		});
+	SECTION("With Abort")
+	{
+		int* Data = nullptr;
+		AutoRTFM::ETransactionResult Result = AutoRTFM::Transact([&]
+			{
+				Data = new int[42]();
+				Data[2] = 42;
+				AutoRTFM::AbortTransaction();
+			});
+
+		REQUIRE(AutoRTFM::ETransactionResult::AbortedByRequest == Result);
+		REQUIRE(nullptr == Data);
+	}
 
-	REQUIRE(42 == Data[2]);
+	SECTION("With Commit")
+	{
+		int* Data = nullptr;
+		AutoRTFM::Commit([&]
+			{
+				Data = new int[42]();
+				Data[2] = 42;
+			});
 
-	delete[] Data;
+		REQUIRE(nullptr != Data);
+		REQUIRE(42 == Data[2]);
+
+		delete[] Data;
+	}
 }
 
-// TODO: Re-enable test once array delete[] works in transactions.
-TEST_CASE("Allocators.ArrayDelete", "[.]")
+TEST_CASE("Allocators.ArrayDelete")
 {
-	int* Data = new int[42]();
+	SECTION("With Abort")
+	{
+		int* Data = new int[42]();
+		Data[10] = 10;
+
+		AutoRTFM::ETransactionResult Result = AutoRTFM::Transact([&]
+			{
+				delete[] Data;
+				AutoRTFM::AbortTransaction();
+			});
+
+		REQUIRE(AutoRTFM::ETransactionResult::AbortedByRequest == Result);
+		REQUIRE(10 == Data[10]);
+	}
 
-	AutoRTFM::Commit([&]
-		{
-			delete[] Data;
-		});
+	SECTION("With Commit")
+	{
+		int* Data = new int[42]();
+
+		AutoRTFM::Commit([&]
+			{
+				delete[] Data;
+			});
+	}
 }
 
 TEST_CASE("Allocators.NewNoOpts")
 {
-	int* Data;
+	SECTION("With Abort")
+	{
+		int* Data = nullptr;
+
+		CLANG_BEGIN_DISABLE_OPTIMIZATIONS
+
+		AutoRTFM::ETransactionResult Result = AutoRTFM::Transact([&]
+			{
+				Data = new int(42);
+				AutoRTFM::AbortTransaction();
+			});
+
+		CLANG_END_DISABLE_OPTIMIZATIONS
+
+		REQUIRE(AutoRTFM::ETransactionResult::AbortedByRequest == Result);
+		REQUIRE(nullptr == Data);
+	}
+
+	SECTION("With Commit")
+	{
+		int* Data = nullptr;
 
-#pragma clang optimize off
-	AutoRTFM::Commit([&]
-		{
-			Data = new int(42);
-		});
-#pragma clang optimize on
+		CLANG_BEGIN_DISABLE_OPTIMIZATIONS
 
-	REQUIRE(42 == *Data);
+		AutoRTFM::Commit([&]
+			{
+				Data = new int(42);
+			});
 
-	delete Data;
+		CLANG_END_DISABLE_OPTIMIZATIONS
+
+		REQUIRE(42 == *Data);
+
+		delete Data;
+	}
 }
 
 TEST_CASE("Allocators.DeleteNoOpts")
 {
-	int* Data = new int(42);
-
-#pragma clang optimize off
-	AutoRTFM::Commit([&]
-		{
-			delete Data;
-		});
-#pragma clang optimize on
+	SECTION("With Abort")
+	{
+		int* Data = new int(42);
+
+		CLANG_BEGIN_DISABLE_OPTIMIZATIONS
+
+		AutoRTFM::ETransactionResult Result = AutoRTFM::Transact([&]
+			{
+				delete Data;
+				AutoRTFM::AbortTransaction();
+			});
+
+		CLANG_END_DISABLE_OPTIMIZATIONS
+
+		REQUIRE(AutoRTFM::ETransactionResult::AbortedByRequest == Result);
+		REQUIRE(42 == *Data);
+	}
+
+	SECTION("With Commit")
+	{
+		int* Data = new int(42);
+
+		CLANG_BEGIN_DISABLE_OPTIMIZATIONS
+
+		AutoRTFM::Commit([&]
+			{
+				delete Data;
+			});
+
+		CLANG_END_DISABLE_OPTIMIZATIONS
+	}
 }
+
+TEST_CASE("Allocators.Free")
+{
+	MSVC_BEGIN_DISABLE_WARN_ALLOC_MAYBE_NULL
+
+	SECTION("With Abort")
+	{
+		void* Data = malloc(32);
+		*reinterpret_cast<int*>(Data) = 42;
+
+		AutoRTFM::ETransactionResult Result = AutoRTFM::Transact([&]
+			{
+				free(Data);
+				AutoRTFM::AbortTransaction();
+			});
+
+		REQUIRE(AutoRTFM::ETransactionResult::AbortedByRequest == Result);
+		REQUIRE(42 == *reinterpret_cast<int*>(Data));
+
+		free(Data);
+	}
+
+	SECTION("With Commit")
+	{
+		void* Data = malloc(32);
+
+		AutoRTFM::Commit([&]
+			{
+				free(Data);
+			});
+	}
+
+	MSVC_END_DISABLE_WARN_ALLOC_MAYBE_NULL
+}
+
+TEST_CASE("Allocators.Malloc")
+{
+	MSVC_BEGIN_DISABLE_WARN_ALLOC_MAYBE_NULL
+
+	SECTION("With Abort")
+	{
+		void* Data = nullptr;
+		AutoRTFM::ETransactionResult Result = AutoRTFM::Transact([&]
+			{
+				Data = malloc(32);
+				*reinterpret_cast<int*>(Data) = 42;
+				AutoRTFM::AbortTransaction();
+			});
+
+		REQUIRE(AutoRTFM::ETransactionResult::AbortedByRequest == Result);
+		REQUIRE(nullptr == Data);
+	}
+
+	SECTION("With Commit")
+	{
+		void* Data = nullptr;
+		AutoRTFM::Commit([&]
+			{
+				Data = malloc(32);
+				*reinterpret_cast<int*>(Data) = 42;
+			});
+
+		REQUIRE(nullptr != Data);
+		REQUIRE(42 == *reinterpret_cast<int*>(Data));
+
+		free(Data);
+	}
+
+	MSVC_END_DISABLE_WARN_ALLOC_MAYBE_NULL
+}
+
+TEST_CASE("Allocators.Calloc")
+{
+	MSVC_BEGIN_DISABLE_WARN_ALLOC_MAYBE_NULL
+
+	SECTION("With Abort")
+	{
+		void* Data = nullptr;
+		AutoRTFM::ETransactionResult Result = AutoRTFM::Transact([&]
+			{
+				Data = calloc(5, sizeof(int));
+				*reinterpret_cast<int*>(Data) = 42;
+				AutoRTFM::AbortTransaction();
+			});
+
+		REQUIRE(AutoRTFM::ETransactionResult::AbortedByRequest == Result);
+		REQUIRE(nullptr == Data);
+	}
+
+	SECTION("With Commit")
+	{
+		void* Data = nullptr;
+		AutoRTFM::Commit([&]
+			{
+				Data = calloc(5, sizeof(int));
+				*reinterpret_cast<int*>(Data) = 42;
+			});
+
+		REQUIRE(nullptr != Data);
+		REQUIRE(42 == reinterpret_cast<int*>(Data)[0]);
+		REQUIRE(0 == reinterpret_cast<int*>(Data)[1]);
+
+		free(Data);
+	}
+
+	MSVC_END_DISABLE_WARN_ALLOC_MAYBE_NULL
+}
+
+TEST_CASE("Allocators.Realloc")
+{
+	MSVC_BEGIN_DISABLE_WARN_ALLOC_MAYBE_NULL
+
+	SECTION("With Abort")
+	{
+		void* Alloc = malloc(32);
+		void* Realloc = nullptr;
+		*reinterpret_cast<int*>(Alloc) = 42;
+
+		AutoRTFM::ETransactionResult Result = AutoRTFM::Transact([&]
+			{
+				Realloc = realloc(Alloc, 64);
+				AutoRTFM::AbortTransaction();
+			});
+
+		REQUIRE(AutoRTFM::ETransactionResult::AbortedByRequest == Result);
+		REQUIRE(nullptr == Realloc);
+		REQUIRE(42 == *reinterpret_cast<int*>(Alloc));
+
+		free(Alloc);
+	}
+
+	SECTION("With Commit")
+	{
+		void* Alloc = malloc(32);
+		void* Realloc = nullptr;
+		*reinterpret_cast<int*>(Alloc) = 42;
+
+		AutoRTFM::Commit([&]
+			{
+				Realloc = realloc(Alloc, 64);
+			});
+
+		REQUIRE(nullptr != Realloc);
+		REQUIRE(42 == *reinterpret_cast<int*>(Realloc));
+
+		free(Realloc);
+	}
+
+	MSVC_END_DISABLE_WARN_ALLOC_MAYBE_NULL
+}
+
+#undef CLANG_BEGIN_DISABLE_OPTIMIZATIONS
+#undef CLANG_END_DISABLE_OPTIMIZATIONS
+#undef MSVC_BEGIN_DISABLE_WARN_ALLOC_MAYBE_NULL
+#undef MSVC_END_DISABLE_WARN_ALLOC_MAYBE_NULL
diff --git a/Engine/Source/Runtime/Core/Private/AutoRTFM/Memcpy.cpp b/Engine/Source/Runtime/Core/Private/AutoRTFM/Memcpy.cpp
--- a/Engine/Source/Runtime/Core/Private/AutoRTFM/Memcpy.cpp
+++ b/Engine/Source/Runtime/Core/Private/AutoRTFM/Memcpy.cpp
@@ -10,7 +10,7 @@ namespace AutoRTFM
 void* MemcpyToNew(void* InDst, const void* InSrc, size_t Size, FContext* Context)
 {
 	UE_LOG(LogAutoRTFM, Verbose, TEXT("MemcpyToNew(%p, %p, %zu)"), InDst, InSrc, Size);
-    AutoRTFM::Unreachable();
+    return memcpy(InDst, InSrc, Size);
 }
 
 void* Memcpy(void* InDst, const void* InSrc, size_t Size, FContext* Context)
diff --git a/Engine/Source/Runtime/Core/Private/AutoRTFM/StdLibSupport.cpp b/Engine/Source/Runtime/Core/Private/AutoRTFM/StdLibSupport.cpp
--- a/Engine/Source/Runtime/Core/Private/AutoRTFM/StdLibSupport.cpp
+++ b/Engine/Source/Runtime/Core/Private/AutoRTFM/StdLibSupport.cpp
@@ -41,6 +41,20 @@ void* RTFM_malloc(size_t Size)
 
 UE_AUTORTFM_REGISTER_OPEN_FUNCTION(malloc);
 
+void* RTFM_calloc(size_t Count, size_t Size)
+{
+    void* Result = calloc(Count, Size);
+	FContext* Context = FContext::Get();
+    Context->GetCurrentTransaction()->DeferUntilAbort([Result]
+    {
+        free(Result);
+    });
+    Context->DidAllocate(Result, Count * Size);
+    return Result;
+}
+
+UE_AUTORTFM_REGISTER_OPEN_FUNCTION(calloc);
+
 void RTFM_free(void* Ptr)
 {
     if (Ptr)
