Fix that StatID's weren't immediately created within a transaction.

diff --git a/Engine/Source/Programs/AutoRTFMTests/AutoRTFMTests.Target.cs b/Engine/Source/Programs/AutoRTFMTests/AutoRTFMTests.Target.cs
--- a/Engine/Source/Programs/AutoRTFMTests/AutoRTFMTests.Target.cs
+++ b/Engine/Source/Programs/AutoRTFMTests/AutoRTFMTests.Target.cs
@@ -42,6 +42,10 @@ public class AutoRTFMTestsTarget : TargetRules
 
 		MinCpuArchX64 = MinimumCpuArchitectureX64.AVX;
 
+		bCompileWithStatsWithoutEngine = true;
+		GlobalDefinitions.Add("ENABLE_STATNAMEDEVENTS=1");
+		GlobalDefinitions.Add("ENABLE_STATNAMEDEVENTS_UOBJECT=1");
+
 		// Allow for disabling writing out the logfile, since in `PreSubmitTest.py` we run this target simultaneously
 		// multiple times, and doing so would cause writing them out to stomp each other.
 		if (!bAllowLogFile)
diff --git a/Engine/Source/Programs/AutoRTFMTests/Private/UECoreTests.cpp b/Engine/Source/Programs/AutoRTFMTests/Private/UECoreTests.cpp
--- a/Engine/Source/Programs/AutoRTFMTests/Private/UECoreTests.cpp
+++ b/Engine/Source/Programs/AutoRTFMTests/Private/UECoreTests.cpp
@@ -8,6 +8,8 @@
 #include "Internationalization/TextCache.h"
 #include "Serialization/CustomVersion.h"
 #include "UObject/NameTypes.h"
+#include "UObject/UObjectArray.h"
+#include "MyAutoRTFMTestObject.h"
 
 TEST_CASE("UECore.FDelegateHandle")
 {
@@ -795,6 +797,166 @@ TEST_CASE("UECore.FTextCache")
 
 				CheckCacheHealthy();
 			}
+}
+
+TEST_CASE("UECore.FUObjectItem")
+{
+	SECTION("CreateStatID First In Open")
+	{
+		FUObjectItem Item;
+		Item.Object = NewObject<UMyAutoRTFMTestObject>();
+		Item.CreateStatID();
+
+		PROFILER_CHAR* const StatIDStringStorage = Item.StatIDStringStorage;
+
+		// If we abort then we won't change anything.
+		AutoRTFM::ETransactionResult Result = AutoRTFM::Transact([&]
+			{
+				Item.CreateStatID();
+				AutoRTFM::AbortTransaction();
+			});
+
+		REQUIRE(AutoRTFM::ETransactionResult::AbortedByRequest == Result);
+		REQUIRE(StatIDStringStorage == Item.StatIDStringStorage);
+
+		// But also if we commit we likewise won't change anything because
+		// the string storage was already created before the transaction
+		// began.
+		Result = AutoRTFM::Transact([&]
+			{
+				Item.CreateStatID();
+			});
+
+		REQUIRE(AutoRTFM::ETransactionResult::Committed == Result);
+		REQUIRE(StatIDStringStorage == Item.StatIDStringStorage);
+	}
+
+	SECTION("CreateStatID First In Closed")
+	{
+		FUObjectItem Item;
+		Item.Object = NewObject<UMyAutoRTFMTestObject>();
+		REQUIRE(nullptr == Item.StatIDStringStorage);
+		REQUIRE(!Item.StatID.IsValidStat());
+
+		// If we abort then we won't change anything.
+		AutoRTFM::ETransactionResult Result = AutoRTFM::Transact([&]
+			{
+				Item.CreateStatID();
+				AutoRTFM::AbortTransaction();
+			});
+
+		REQUIRE(AutoRTFM::ETransactionResult::AbortedByRequest == Result);
+		REQUIRE(nullptr != Item.StatIDStringStorage);
+		REQUIRE(Item.StatID.IsValidStat());
+
+		PROFILER_CHAR* const StatIDStringStorage = Item.StatIDStringStorage;
+
+		// If we commit though we'll create the stat ID.
+		Result = AutoRTFM::Transact([&]
+			{
+				Item.CreateStatID();
+			});
+
+		REQUIRE(AutoRTFM::ETransactionResult::Committed == Result);
+		REQUIRE(StatIDStringStorage == Item.StatIDStringStorage);
+		REQUIRE(Item.StatID.IsValidStat());
+	}
+
+	SECTION("CreateStatID On In-Transaction Object")
+	{
+		AutoRTFM::ETransactionResult Result = AutoRTFM::Transact([&]
+			{
+				FUObjectItem Item;
+				Item.Object = NewObject<UMyAutoRTFMTestObject>();
+				Item.CreateStatID();
+
+				AutoRTFM::Open([&]
+					{
+						REQUIRE(nullptr != Item.StatIDStringStorage);
+						REQUIRE(Item.StatID.IsValidStat());
+					});
+
+				AutoRTFM::AbortTransaction();
+			});
+
+		REQUIRE(AutoRTFM::ETransactionResult::AbortedByRequest == Result);
+
+		Result = AutoRTFM::Transact([&]
+			{
+				FUObjectItem Item;
+				Item.Object = NewObject<UMyAutoRTFMTestObject>();
+				Item.CreateStatID();
+
+				AutoRTFM::Open([&]
+					{
+						REQUIRE(nullptr != Item.StatIDStringStorage);
+						REQUIRE(Item.StatID.IsValidStat());
+					});
+			});
+
+		REQUIRE(AutoRTFM::ETransactionResult::Committed == Result);
+	}
+
+	SECTION("CreateStatID In Closed Then Again In Open")
+	{
+		{
+			FUObjectItem Item;
+			Item.Object = NewObject<UMyAutoRTFMTestObject>();
+			REQUIRE(nullptr == Item.StatIDStringStorage);
+			REQUIRE(!Item.StatID.IsValidStat());
+
+			AutoRTFM::ETransactionResult Result = AutoRTFM::Transact([&]
+				{
+					Item.CreateStatID();
+
+					AutoRTFM::Open([&]
+						{
+							REQUIRE(nullptr != Item.StatIDStringStorage);
+							REQUIRE(Item.StatID.IsValidStat());
+
+							PROFILER_CHAR* const StatIDStringStorage = Item.StatIDStringStorage;
+
+							Item.CreateStatID();
+
+							REQUIRE(StatIDStringStorage == Item.StatIDStringStorage);
+							REQUIRE(Item.StatID.IsValidStat());
+						});
+
+					AutoRTFM::AbortTransaction();
+				});
+
+			REQUIRE(AutoRTFM::ETransactionResult::AbortedByRequest == Result);
+			REQUIRE(nullptr != Item.StatIDStringStorage);
+			REQUIRE(Item.StatID.IsValidStat());
+		}
+
+		{
+			FUObjectItem Item;
+			Item.Object = NewObject<UMyAutoRTFMTestObject>();
+			REQUIRE(nullptr == Item.StatIDStringStorage);
+			REQUIRE(!Item.StatID.IsValidStat());
+
+			AutoRTFM::ETransactionResult Result = AutoRTFM::Transact([&]
+				{
+					Item.CreateStatID();
+
+					AutoRTFM::Open([&]
+						{
+							REQUIRE(nullptr != Item.StatIDStringStorage);
+							REQUIRE(Item.StatID.IsValidStat());
+
+							PROFILER_CHAR* const StatIDStringStorage = Item.StatIDStringStorage;
+
+							Item.CreateStatID();
+
+							REQUIRE(StatIDStringStorage == Item.StatIDStringStorage);
+							REQUIRE(Item.StatID.IsValidStat());
+						});
+				});
+
+			REQUIRE(AutoRTFM::ETransactionResult::Committed == Result);
+			REQUIRE(nullptr != Item.StatIDStringStorage);
+			REQUIRE(Item.StatID.IsValidStat());
 		}
 	}
 }
diff --git a/Engine/Source/Runtime/CoreUObject/Private/UObject/UObjectArray.cpp b/Engine/Source/Runtime/CoreUObject/Private/UObject/UObjectArray.cpp
--- a/Engine/Source/Runtime/CoreUObject/Private/UObject/UObjectArray.cpp
+++ b/Engine/Source/Runtime/CoreUObject/Private/UObject/UObjectArray.cpp
@@ -67,22 +67,21 @@ void FUObjectItem::CreateStatID() const
 #if STATS
 	StatID = FDynamicStats::CreateStatId<FStatGroup_STATGROUP_UObjects>(LongName);
 #else // ENABLE_STATNAMEDEVENTS
-	const auto& ConversionData = StringCast<PROFILER_CHAR>(*LongName);
-	const int32 NumStorageChars = (ConversionData.Length() + 1);	//length doesn't include null terminator
+	UE_AUTORTFM_OPEN(
+		{
+			const auto& ConversionData = StringCast<PROFILER_CHAR>(*LongName);
+			const int32 NumStorageChars = (ConversionData.Length() + 1);	//length doesn't include null terminator
 
-	PROFILER_CHAR* StoragePtr = new PROFILER_CHAR[NumStorageChars];
-	FMemory::Memcpy(StoragePtr, ConversionData.Get(), NumStorageChars * sizeof(PROFILER_CHAR));
+			PROFILER_CHAR* const StoragePtr = new PROFILER_CHAR[NumStorageChars];
+			FMemory::Memcpy(StoragePtr, ConversionData.Get(), NumStorageChars * sizeof(PROFILER_CHAR));
 
-	// delay the delete of the StatIDStringStorage swap until after the transaction
-	UE_AUTORTFM_ONCOMMIT2(=)
-	{
-		if (FPlatformAtomics::InterlockedCompareExchangePointer((void**)&StatIDStringStorage, StoragePtr, nullptr) != nullptr)
-		{
-			delete[] StoragePtr;
-		}
-	};
+			if (FPlatformAtomics::InterlockedCompareExchangePointer((void**)&StatIDStringStorage, StoragePtr, nullptr) != nullptr)
+			{
+				delete[] StoragePtr;
+			}
 
-	StatID = TStatId(StatIDStringStorage);
+			StatID = TStatId(StatIDStringStorage);
+		});
 #endif
 }
 #endif
