Fix the build breakage caused by a bad merge.

#rnx

diff --git a/Engine/Source/Programs/AutoRTFMTests/Private/UECoreTests.cpp b/Engine/Source/Programs/AutoRTFMTests/Private/UECoreTests.cpp
--- a/Engine/Source/Programs/AutoRTFMTests/Private/UECoreTests.cpp
+++ b/Engine/Source/Programs/AutoRTFMTests/Private/UECoreTests.cpp
@@ -797,6 +797,8 @@ TEST_CASE("UECore.FTextCache")
 
 				CheckCacheHealthy();
 			}
+		}
+	}
 }
 
 TEST_CASE("UECore.FUObjectItem")
