[AutoRTFM] Various string stdlib fixes and add more tests.

* Fix off-by-one error in strcpy() memory record.
* Add a number of missing UE_AUTORTFM_REGISTER_ declarations for string functions.
* Add missing REQUIRE checks for aborted transaction results.
* Track end-pointer write of strtol().
* Track writes of printf-to-string functions.
* Throw if any format specifiers hold '%n' when called in a transaction.

diff --git a/Engine/Source/Programs/AutoRTFMTests/Private/CString.cpp b/Engine/Source/Programs/AutoRTFMTests/Private/CString.cpp
--- a/Engine/Source/Programs/AutoRTFMTests/Private/CString.cpp
+++ b/Engine/Source/Programs/AutoRTFMTests/Private/CString.cpp
@@ -1,15 +1,27 @@
 // Copyright Epic Games, Inc. All Rights Reserved.
 
 #include "Catch2Includes.h"
+#include "CoreGlobals.h"
+#include "Logging/LogVerbosity.h"
+#include "Misc/FeedbackContext.h"
 #include <AutoRTFM/AutoRTFM.h>
+#include <catch_amalgamated.hpp>
 #include <cstring>
 
 #ifdef __clang__
-#define CLANG_BEGIN_DISABLE_OPTIMIZATIONS  _Pragma("clang optimize off")
-#define CLANG_END_DISABLE_OPTIMIZATIONS    _Pragma("clang optimize on")
+#define CLANG_BEGIN_DISABLE_OPTIMIZATIONS _Pragma("clang optimize off")
+#define CLANG_END_DISABLE_OPTIMIZATIONS   _Pragma("clang optimize on")
+#define CLANG_BEGIN_DISABLE_WARN_FORMAT_TRUNCATED \
+	_Pragma("clang diagnostic push") \
+	_Pragma("clang diagnostic ignored \"-Wunknown-warning-option\"") \
+	_Pragma("clang diagnostic ignored \"-Wformat-truncation\"")
+#define CLANG_END_DISABLE_WARN_FORMAT_TRUNCATED \
+	_Pragma("clang diagnostic pop")
 #else
 #define CLANG_BEGIN_DISABLE_OPTIMIZATIONS
 #define CLANG_END_DISABLE_OPTIMIZATIONS
+#define CLANG_BEGIN_DISABLE_WARN_FORMAT_TRUNCATED
+#define CLANG_END_DISABLE_WARN_FORMAT_TRUNCATED
 #endif
 
 #ifdef _MSC_VER
@@ -22,6 +34,98 @@
 #define MSVC_END_DISABLE_UNSAFE_FN
 #endif
 
+namespace
+{
+
+// The expected AutoRTFM warning message when attempting to printf() using a
+// format string containing a '%n' format specifier.
+static constexpr auto* kPercentNWarning = TEXT("AutoRTFM does not support format strings containing '%n'");
+
+// CheckWCharFunction calls Function with WChar outside a transaction, in an
+// aborted transaction and committed transaction. The value returned when called
+// in the committed transation is expected to match the value returned when
+// called outside the transaction.
+// Function must be the deterministic function with the signature 'int(wchar_t)'.
+template<typename FUNCTION>
+void CheckWCharFunction(FUNCTION&& Function, wchar_t WChar)
+{
+	int Expect = Function(WChar);
+
+	SECTION("With Abort")
+	{
+		int Got = 0;
+
+		CLANG_BEGIN_DISABLE_OPTIMIZATIONS
+		AutoRTFM::ETransactionResult Result = AutoRTFM::Transact([&]
+		{
+			Got = Function(WChar);
+			AutoRTFM::AbortTransaction();
+		});
+		CLANG_END_DISABLE_OPTIMIZATIONS
+
+		REQUIRE(AutoRTFM::ETransactionResult::AbortedByRequest == Result);
+		REQUIRE(0 == Got);
+	}
+
+	SECTION("With Commit")
+	{
+		int Got = 0;
+
+		CLANG_BEGIN_DISABLE_OPTIMIZATIONS
+		AutoRTFM::Commit([&] { Got = Function(WChar); });
+		CLANG_END_DISABLE_OPTIMIZATIONS
+
+		REQUIRE(Expect == Got);
+	}
+}
+
+// A helper class that for the lifetime of the object, intercepts and records
+// UE_LOG warnings.
+class FCaptureWarningContext : private FFeedbackContext
+{
+public:
+	FCaptureWarningContext() : OldContext(GWarn) { GWarn = this; }
+	~FCaptureWarningContext() { GWarn = OldContext; }
+
+	void Serialize(const TCHAR* V, ELogVerbosity::Type Verbosity, const FName& Category) override
+	{
+		if (Verbosity == ELogVerbosity::Warning)
+		{
+			Warnings.push_back(V);
+		}
+		else
+		{
+			OldContext->Serialize(V, Verbosity, Category);
+		}
+	}
+
+	void Serialize(const TCHAR* V, ELogVerbosity::Type Verbosity, const FName& Category, double Time) override
+	{
+		if (Verbosity == ELogVerbosity::Warning)
+		{
+			Warnings.push_back(V);
+		}
+		else
+		{
+			OldContext->Serialize(V, Verbosity, Category, Time);
+		}
+	}
+
+	const std::vector<FString>& GetWarnings() const
+	{
+		return Warnings;
+	}
+
+private:
+	FCaptureWarningContext(const FCaptureWarningContext&) = delete;
+	FCaptureWarningContext& operator = (const FCaptureWarningContext&) = delete;
+
+	FFeedbackContext* OldContext = nullptr;
+	std::vector<FString> Warnings;
+};
+
+}  // anonymous namespace
+
 TEST_CASE("CString.memcpy")
 {
 	const char* const From = "Kittie says meow";
@@ -84,8 +188,8 @@ TEST_CASE("CString.strcpy")
 	MSVC_BEGIN_DISABLE_WARN_UNSAFE_FUNCTION
 
 	const char From[] = "Kittie says meow";
-	char To[] = "Doggie says woof";
-	static_assert(sizeof(From) == sizeof(To));
+	char To[] = "Doggie says woof____";
+	static_assert(sizeof(From) <= sizeof(To));
 
 	SECTION("With Abort")
 	{
@@ -98,7 +202,7 @@ TEST_CASE("CString.strcpy")
 		CLANG_END_DISABLE_OPTIMIZATIONS
 
 		REQUIRE(AutoRTFM::ETransactionResult::AbortedByRequest == Result);
-		REQUIRE("Doggie says woof" == std::string_view(To));
+		REQUIRE("Doggie says woof____" == std::string_view(To));
 	}
 
 	SECTION("With Commit")
@@ -229,6 +333,7 @@ TEST_CASE("CString.memcmp")
 		});
 		CLANG_END_DISABLE_OPTIMIZATIONS
 
+		REQUIRE(AutoRTFM::ETransactionResult::AbortedByRequest == Result);
 		REQUIRE(0 == Compare);
 	}
 
@@ -260,6 +365,7 @@ TEST_CASE("CString.strcmp")
 		});
 		CLANG_END_DISABLE_OPTIMIZATIONS
 
+		REQUIRE(AutoRTFM::ETransactionResult::AbortedByRequest == Result);
 		REQUIRE(0 == Compare);
 	}
 
@@ -291,6 +397,7 @@ TEST_CASE("CString.strncmp")
 		});
 		CLANG_END_DISABLE_OPTIMIZATIONS
 
+		REQUIRE(AutoRTFM::ETransactionResult::AbortedByRequest == Result);
 		REQUIRE(0 == Compare);
 	}
 
@@ -322,6 +429,7 @@ TEST_CASE("CString.strchr")
 		});
 		CLANG_END_DISABLE_OPTIMIZATIONS
 
+		REQUIRE(AutoRTFM::ETransactionResult::AbortedByRequest == Result);
 		REQUIRE(nullptr == Value);
 	}
 
@@ -353,6 +461,7 @@ TEST_CASE("CString.strrchr")
 		});
 		CLANG_END_DISABLE_OPTIMIZATIONS
 
+		REQUIRE(AutoRTFM::ETransactionResult::AbortedByRequest == Result);
 		REQUIRE(nullptr == Value);
 	}
 
@@ -384,6 +493,7 @@ TEST_CASE("CString.strstr")
 		});
 		CLANG_END_DISABLE_OPTIMIZATIONS
 
+		REQUIRE(AutoRTFM::ETransactionResult::AbortedByRequest == Result);
 		REQUIRE(nullptr == Value);
 	}
 
@@ -415,6 +525,7 @@ TEST_CASE("CString.strlen")
 		});
 		CLANG_END_DISABLE_OPTIMIZATIONS
 
+		REQUIRE(AutoRTFM::ETransactionResult::AbortedByRequest == Result);
 		REQUIRE(0 == Value);
 	}
 
@@ -430,6 +541,881 @@ TEST_CASE("CString.strlen")
 	}
 }
 
+TEST_CASE("CString.strtol")
+{
+	const char* Str = "123xyz";
+
+	SECTION("With end pointer")
+	{
+		SECTION("With Abort")
+		{
+			long int Value = 0;
+			char* EndPtr = nullptr;
+
+			CLANG_BEGIN_DISABLE_OPTIMIZATIONS
+			AutoRTFM::ETransactionResult Result = AutoRTFM::Transact([&]
+			{
+				Value = strtol(Str, &EndPtr, 10);
+				AutoRTFM::AbortTransaction();
+			});
+			CLANG_END_DISABLE_OPTIMIZATIONS
+
+			REQUIRE(AutoRTFM::ETransactionResult::AbortedByRequest == Result);
+			REQUIRE(0 == Value);
+			REQUIRE(nullptr == EndPtr);
+		}
+
+		SECTION("With Commit")
+		{
+			long int Value = 0;
+			char* EndPtr = nullptr;
+
+			CLANG_BEGIN_DISABLE_OPTIMIZATIONS
+			AutoRTFM::Commit([&] { Value = strtol(Str, &EndPtr, 10); });
+			CLANG_END_DISABLE_OPTIMIZATIONS
+
+			REQUIRE(123 == Value);
+			REQUIRE((Str + 3) == EndPtr);
+		}
+	}
+
+	SECTION("Without end pointer")
+	{
+		SECTION("With Abort")
+		{
+			long int Value = 0;
+
+			CLANG_BEGIN_DISABLE_OPTIMIZATIONS
+			AutoRTFM::ETransactionResult Result = AutoRTFM::Transact([&]
+			{
+				Value = strtol(Str, nullptr, 10);
+				AutoRTFM::AbortTransaction();
+			});
+			CLANG_END_DISABLE_OPTIMIZATIONS
+
+			REQUIRE(AutoRTFM::ETransactionResult::AbortedByRequest == Result);
+			REQUIRE(0 == Value);
+		}
+
+		SECTION("With Commit")
+		{
+			long int Value = 0;
+
+			CLANG_BEGIN_DISABLE_OPTIMIZATIONS
+			AutoRTFM::Commit([&] { Value = strtol(Str, nullptr, 10); });
+			CLANG_END_DISABLE_OPTIMIZATIONS
+
+			REQUIRE(123 == Value);
+		}
+	}
+}
+
+
+TEST_CASE("CString.wcscpy")
+{
+	MSVC_BEGIN_DISABLE_WARN_UNSAFE_FUNCTION
+
+	const wchar_t From[] = L"Kittie says meow";
+	wchar_t To[] = L"Doggie says woof____";
+	static_assert(sizeof(From) <= sizeof(To));
+
+	SECTION("With Abort")
+	{
+		CLANG_BEGIN_DISABLE_OPTIMIZATIONS
+		AutoRTFM::ETransactionResult Result = AutoRTFM::Transact([&]
+		{
+			wcscpy(To, From);
+			AutoRTFM::AbortTransaction();
+		});
+		CLANG_END_DISABLE_OPTIMIZATIONS
+
+		REQUIRE(AutoRTFM::ETransactionResult::AbortedByRequest == Result);
+		REQUIRE(L"Doggie says woof____" == std::wstring_view(To));
+	}
+
+	SECTION("With Commit")
+	{
+		CLANG_BEGIN_DISABLE_OPTIMIZATIONS
+		AutoRTFM::Commit([&] { wcscpy(To, From); });
+		CLANG_END_DISABLE_OPTIMIZATIONS
+
+		REQUIRE(L"Kittie says meow" == std::wstring_view(To));
+	}
+
+	MSVC_END_DISABLE_UNSAFE_FN
+}
+
+TEST_CASE("CString.wcsncpy")
+{
+	MSVC_BEGIN_DISABLE_WARN_UNSAFE_FUNCTION
+
+	const wchar_t* const From = L"Kittie says meow";
+	wchar_t To[] = L"Doggie says woof";
+
+	SECTION("With Abort")
+	{
+		CLANG_BEGIN_DISABLE_OPTIMIZATIONS
+		AutoRTFM::ETransactionResult Result = AutoRTFM::Transact([&]
+		{
+			wcsncpy(To, From, 6);
+			AutoRTFM::AbortTransaction();
+		});
+		CLANG_END_DISABLE_OPTIMIZATIONS
+
+		REQUIRE(AutoRTFM::ETransactionResult::AbortedByRequest == Result);
+		REQUIRE(L"Doggie says woof" == std::wstring_view(To));
+	}
+
+	SECTION("With Commit")
+	{
+		CLANG_BEGIN_DISABLE_OPTIMIZATIONS
+		AutoRTFM::Commit([&] { wcsncpy(To, From, 6); });
+		CLANG_END_DISABLE_OPTIMIZATIONS
+
+		REQUIRE(L"Kittie says woof" == std::wstring_view(To));
+	}
+
+	MSVC_END_DISABLE_UNSAFE_FN
+}
+
+TEST_CASE("CString.wcsncmp")
+{
+	const wchar_t* A = L"This";
+
+	SECTION("With Abort")
+	{
+		int Compare = 0;
+
+		CLANG_BEGIN_DISABLE_OPTIMIZATIONS
+		AutoRTFM::ETransactionResult Result = AutoRTFM::Transact([&]
+		{
+			Compare = wcsncmp(A, L"That", 3);
+			AutoRTFM::AbortTransaction();
+		});
+		CLANG_END_DISABLE_OPTIMIZATIONS
+
+		REQUIRE(AutoRTFM::ETransactionResult::AbortedByRequest == Result);
+		REQUIRE(0 == Compare);
+	}
+
+	SECTION("With Commit")
+	{
+		int Compare = 0;
+
+		CLANG_BEGIN_DISABLE_OPTIMIZATIONS
+		AutoRTFM::Commit([&] { Compare = wcsncmp(A, L"That", 3); });
+		CLANG_END_DISABLE_OPTIMIZATIONS
+
+		REQUIRE(0 < Compare);
+	}
+}
+
+TEST_CASE("CString.wcschr")
+{
+	SECTION("Const")
+	{
+		const wchar_t* A = L"Thinking";
+
+		SECTION("With Abort")
+		{
+			const wchar_t* Value = nullptr;
+
+			CLANG_BEGIN_DISABLE_OPTIMIZATIONS
+			AutoRTFM::ETransactionResult Result = AutoRTFM::Transact([&]
+			{
+				Value = wcschr(A, L'i');
+				AutoRTFM::AbortTransaction();
+			});
+			CLANG_END_DISABLE_OPTIMIZATIONS
+
+			REQUIRE(AutoRTFM::ETransactionResult::AbortedByRequest == Result);
+			REQUIRE(nullptr == Value);
+		}
+
+		SECTION("With Commit")
+		{
+			const wchar_t* Value = nullptr;
+
+			CLANG_BEGIN_DISABLE_OPTIMIZATIONS
+			AutoRTFM::Commit([&] { Value = wcschr(A, L'i'); });
+			CLANG_END_DISABLE_OPTIMIZATIONS
+
+			REQUIRE((A + 2) == Value);
+		}
+	}
+
+	SECTION("Non-Const")
+	{
+		wchar_t A[] = L"Thinking";
+
+		SECTION("With Abort")
+		{
+			wchar_t* Value = nullptr;
+
+			CLANG_BEGIN_DISABLE_OPTIMIZATIONS
+			AutoRTFM::ETransactionResult Result = AutoRTFM::Transact([&]
+			{
+				Value = wcschr(A, L'i');
+				AutoRTFM::AbortTransaction();
+			});
+			CLANG_END_DISABLE_OPTIMIZATIONS
+
+			REQUIRE(AutoRTFM::ETransactionResult::AbortedByRequest == Result);
+			REQUIRE(nullptr == Value);
+		}
+
+		SECTION("With Commit")
+		{
+			wchar_t* Value = nullptr;
+
+			CLANG_BEGIN_DISABLE_OPTIMIZATIONS
+			AutoRTFM::Commit([&] { Value = wcschr(A, L'i'); });
+			CLANG_END_DISABLE_OPTIMIZATIONS
+
+			REQUIRE((A + 2) == Value);
+		}
+	}
+}
+
+TEST_CASE("CString.wcsstr")
+{
+	SECTION("Const")
+	{
+		const wchar_t* A = L"Thinking";
+
+		SECTION("With Abort")
+		{
+			const wchar_t* Value = nullptr;
+
+			CLANG_BEGIN_DISABLE_OPTIMIZATIONS
+			AutoRTFM::ETransactionResult Result = AutoRTFM::Transact([&]
+			{
+				Value = wcsstr(A, L"ink");
+				AutoRTFM::AbortTransaction();
+			});
+			CLANG_END_DISABLE_OPTIMIZATIONS
+
+			REQUIRE(AutoRTFM::ETransactionResult::AbortedByRequest == Result);
+			REQUIRE(nullptr == Value);
+		}
+
+		SECTION("With Commit")
+		{
+			const wchar_t* Value = nullptr;
+
+			CLANG_BEGIN_DISABLE_OPTIMIZATIONS
+			AutoRTFM::Commit([&] { Value = wcsstr(A, L"ink"); });
+			CLANG_END_DISABLE_OPTIMIZATIONS
+
+			REQUIRE((A + 2) == Value);
+		}
+	}
+
+	SECTION("Non-Const")
+	{
+		wchar_t A[] = L"Thinking";
+
+		SECTION("With Abort")
+		{
+			wchar_t* Value = nullptr;
+
+			CLANG_BEGIN_DISABLE_OPTIMIZATIONS
+			AutoRTFM::ETransactionResult Result = AutoRTFM::Transact([&]
+			{
+				Value = wcsstr(A, L"ink");
+				AutoRTFM::AbortTransaction();
+			});
+			CLANG_END_DISABLE_OPTIMIZATIONS
+
+			REQUIRE(AutoRTFM::ETransactionResult::AbortedByRequest == Result);
+			REQUIRE(nullptr == Value);
+		}
+
+		SECTION("With Commit")
+		{
+			wchar_t* Value = nullptr;
+
+			CLANG_BEGIN_DISABLE_OPTIMIZATIONS
+			AutoRTFM::Commit([&] { Value = wcsstr(A, L"ink"); });
+			CLANG_END_DISABLE_OPTIMIZATIONS
+
+			REQUIRE((A + 2) == Value);
+		}
+	}
+}
+
+TEST_CASE("CString.wcscmp")
+{
+	const wchar_t* A = L"This";
+
+	SECTION("With Abort")
+	{
+		int Compare = 0;
+
+		CLANG_BEGIN_DISABLE_OPTIMIZATIONS
+		AutoRTFM::ETransactionResult Result = AutoRTFM::Transact([&]
+		{
+			Compare = wcscmp(A, L"That");
+			AutoRTFM::AbortTransaction();
+		});
+		CLANG_END_DISABLE_OPTIMIZATIONS
+
+		REQUIRE(AutoRTFM::ETransactionResult::AbortedByRequest == Result);
+		REQUIRE(0 == Compare);
+	}
+
+	SECTION("With Commit")
+	{
+		int Compare = 0;
+
+		CLANG_BEGIN_DISABLE_OPTIMIZATIONS
+		AutoRTFM::Commit([&] { Compare = wcscmp(A, L"That"); });
+		CLANG_END_DISABLE_OPTIMIZATIONS
+
+		REQUIRE(0 < Compare);
+	}
+}
+
+TEST_CASE("CString.wcslen")
+{
+	const wchar_t* A = L"This";
+
+	SECTION("With Abort")
+	{
+		size_t Value = 0;
+
+		CLANG_BEGIN_DISABLE_OPTIMIZATIONS
+		AutoRTFM::ETransactionResult Result = AutoRTFM::Transact([&]
+		{
+			Value = wcslen(A);
+			AutoRTFM::AbortTransaction();
+		});
+		CLANG_END_DISABLE_OPTIMIZATIONS
+
+		REQUIRE(AutoRTFM::ETransactionResult::AbortedByRequest == Result);
+		REQUIRE(0 == Value);
+	}
+
+	SECTION("With Commit")
+	{
+		size_t Value = 0;
+
+		CLANG_BEGIN_DISABLE_OPTIMIZATIONS
+		AutoRTFM::Commit([&] { Value = wcslen(A); });
+		CLANG_END_DISABLE_OPTIMIZATIONS
+
+		REQUIRE(4 == Value);
+	}
+}
+
+
+TEST_CASE("CString.iswupper")
+{
+	CheckWCharFunction(iswupper, L'X');
+	CheckWCharFunction(iswupper, L'x');
+}
+
+TEST_CASE("CString.iswlower")
+{
+	CheckWCharFunction(iswlower, L'X');
+	CheckWCharFunction(iswlower, L'x');
+}
+
+TEST_CASE("CString.iswalpha")
+{
+	CheckWCharFunction(iswalpha, L'X');
+	CheckWCharFunction(iswalpha, L'5');
+}
+
+TEST_CASE("CString.iswgraph")
+{
+	CheckWCharFunction(iswgraph, L'X');
+	CheckWCharFunction(iswgraph, L'x');
+}
+
+TEST_CASE("CString.iswprint")
+{
+	CheckWCharFunction(iswprint, L'X');
+	CheckWCharFunction(iswprint, L'x');
+}
+
+TEST_CASE("CString.iswpunct")
+{
+	CheckWCharFunction(iswpunct, L'X');
+	CheckWCharFunction(iswpunct, L'^');
+}
+
+TEST_CASE("CString.iswalnum")
+{
+	CheckWCharFunction(iswalnum, L'X');
+	CheckWCharFunction(iswalnum, L'6');
+	CheckWCharFunction(iswalnum, L' ');
+}
+
+TEST_CASE("CString.iswdigit")
+{
+	CheckWCharFunction(iswdigit, L'X');
+	CheckWCharFunction(iswdigit, L'5');
+}
+
+TEST_CASE("CString.iswxdigit")
+{
+	CheckWCharFunction(iswxdigit, L'X');
+	CheckWCharFunction(iswxdigit, L'q');
+}
+
+TEST_CASE("CString.iswspace")
+{
+	CheckWCharFunction(iswspace, L'X');
+	CheckWCharFunction(iswspace, L' ');
+}
+
+TEST_CASE("CString.iswcntrl")
+{
+	CheckWCharFunction(iswcntrl, L'X');
+	CheckWCharFunction(iswcntrl, L'\u2028');
+}
+
+TEST_CASE("CString.snprintf")
+{
+	static constexpr size_t BufferSize = 64;
+	char Buffer[BufferSize] = "_____________________";
+
+	SECTION("Fits in buffer")
+	{
+		SECTION("With Abort")
+		{
+			int Count = 0;
+
+			CLANG_BEGIN_DISABLE_OPTIMIZATIONS
+			AutoRTFM::ETransactionResult Result = AutoRTFM::Transact([&]
+			{
+				Count = snprintf(Buffer, BufferSize, "cat says '%s'!", "meow");
+				AutoRTFM::AbortTransaction();
+			});
+			CLANG_END_DISABLE_OPTIMIZATIONS
+
+			REQUIRE(AutoRTFM::ETransactionResult::AbortedByRequest == Result);
+			REQUIRE(0 == Count);
+			REQUIRE("_____________________" == std::string_view(Buffer));
+		}
+
+		SECTION("With Commit")
+		{
+			int Count = 0;
+
+			CLANG_BEGIN_DISABLE_OPTIMIZATIONS
+			AutoRTFM::Commit([&]
+			{
+				Count = snprintf(Buffer, BufferSize, "cat says '%s'!", "meow");
+			});
+			CLANG_END_DISABLE_OPTIMIZATIONS
+
+			REQUIRE(16 == Count);
+			REQUIRE("cat says 'meow'!" == std::string_view(Buffer));
+		}
+	}
+
+	SECTION("Buffer too small")
+	{
+		CLANG_BEGIN_DISABLE_WARN_FORMAT_TRUNCATED
+
+		SECTION("With Abort")
+		{
+			int Count = 0;
+
+			CLANG_BEGIN_DISABLE_OPTIMIZATIONS
+			AutoRTFM::ETransactionResult Result = AutoRTFM::Transact([&]
+			{
+				Count = snprintf(Buffer, 8, "cat says '%s'!", "meow");
+				AutoRTFM::AbortTransaction();
+			});
+			CLANG_END_DISABLE_OPTIMIZATIONS
+
+			REQUIRE(AutoRTFM::ETransactionResult::AbortedByRequest == Result);
+			REQUIRE(0 == Count);
+			REQUIRE("_____________________" == std::string_view(Buffer));
+		}
+
+		SECTION("With Commit")
+		{
+			int Count = 0;
+
+			CLANG_BEGIN_DISABLE_OPTIMIZATIONS
+			AutoRTFM::Commit([&]
+			{
+				Count = snprintf(Buffer, 8, "cat says '%s'!", "meow");
+			});
+			CLANG_END_DISABLE_OPTIMIZATIONS
+
+			REQUIRE(16 == Count);
+			REQUIRE("cat say" == std::string_view(Buffer));
+		}
+
+		CLANG_END_DISABLE_WARN_FORMAT_TRUNCATED
+	}
+
+	SECTION("Null buffer")
+	{
+		SECTION("With Abort")
+		{
+			int Count = 0;
+
+			CLANG_BEGIN_DISABLE_OPTIMIZATIONS
+			AutoRTFM::ETransactionResult Result = AutoRTFM::Transact([&]
+			{
+				Count = snprintf(nullptr, 0, "cat says '%s'!", "meow");
+				AutoRTFM::AbortTransaction();
+			});
+			CLANG_END_DISABLE_OPTIMIZATIONS
+
+			REQUIRE(AutoRTFM::ETransactionResult::AbortedByRequest == Result);
+			REQUIRE(0 == Count);
+		}
+
+		SECTION("With Commit")
+		{
+			int Count = 0;
+
+			CLANG_BEGIN_DISABLE_OPTIMIZATIONS
+			AutoRTFM::Commit([&]
+			{
+				Count = snprintf(nullptr, 0, "cat says '%s'!", "meow");
+			});
+			CLANG_END_DISABLE_OPTIMIZATIONS
+
+			REQUIRE(16 == Count);
+		}
+	}
+
+	SECTION("PercentN")
+	{
+		FCaptureWarningContext WarningContext;
+
+		int Count = 0;
+
+		CLANG_BEGIN_DISABLE_OPTIMIZATIONS
+		AutoRTFM::ETransactionResult Result = AutoRTFM::Transact([&]
+		{
+			snprintf(Buffer, 8, "%n", &Count);
+		});
+		CLANG_END_DISABLE_OPTIMIZATIONS
+
+		REQUIRE(AutoRTFM::ETransactionResult::AbortedByLanguage == Result);
+		REQUIRE(0 == Count);
+		REQUIRE("_____________________" == std::string_view(Buffer));
+		REQUIRE_THAT(WarningContext.GetWarnings(),
+			Catch::Matchers::VectorContains(FString(kPercentNWarning)));
+	}
+
+	SECTION("PercentPercentN")
+	{
+		SECTION("With Abort")
+		{
+			int Count = 0;
+
+			CLANG_BEGIN_DISABLE_OPTIMIZATIONS
+			AutoRTFM::ETransactionResult Result = AutoRTFM::Transact([&]
+			{
+				Count = snprintf(Buffer, 8, "%%n");
+				AutoRTFM::AbortTransaction();
+			});
+			CLANG_END_DISABLE_OPTIMIZATIONS
+
+			REQUIRE(AutoRTFM::ETransactionResult::AbortedByRequest == Result);
+			REQUIRE(0 == Count);
+			REQUIRE("_____________________" == std::string_view(Buffer));
+		}
+
+		SECTION("With Commit")
+		{
+			int Count = 0;
+
+			CLANG_BEGIN_DISABLE_OPTIMIZATIONS
+			AutoRTFM::Commit([&]
+			{
+				Count = snprintf(Buffer, 8, "%%n");
+			});
+			CLANG_END_DISABLE_OPTIMIZATIONS
+
+			REQUIRE(2 == Count);
+			REQUIRE("%n" == std::string_view(Buffer));
+		}
+	}
+}
+
+TEST_CASE("CString.swprintf")
+{
+	static constexpr size_t BufferSize = 64;
+	wchar_t Buffer[BufferSize] = L"_____________________";
+
+	SECTION("Fits in buffer")
+	{
+		SECTION("With Abort")
+		{
+			int Count = 0;
+
+			CLANG_BEGIN_DISABLE_OPTIMIZATIONS
+			AutoRTFM::ETransactionResult Result = AutoRTFM::Transact([&]
+			{
+				Count = swprintf(Buffer, BufferSize, L"cat says '%ls'!", L"meow");
+				AutoRTFM::AbortTransaction();
+			});
+			CLANG_END_DISABLE_OPTIMIZATIONS
+
+			REQUIRE(AutoRTFM::ETransactionResult::AbortedByRequest == Result);
+			REQUIRE(0 == Count);
+			REQUIRE(L"_____________________" == std::wstring_view(Buffer));
+		}
+
+		SECTION("With Commit")
+		{
+			int Count = 0;
+
+			CLANG_BEGIN_DISABLE_OPTIMIZATIONS
+			AutoRTFM::Commit([&]
+			{
+				Count = swprintf(Buffer, BufferSize, L"cat says '%ls'!", L"meow");
+			});
+			CLANG_END_DISABLE_OPTIMIZATIONS
+
+			REQUIRE(16 == Count);
+			REQUIRE(L"cat says 'meow'!" == std::wstring_view(Buffer));
+		}
+	}
+
+	SECTION("Buffer too small")
+	{
+		CLANG_BEGIN_DISABLE_WARN_FORMAT_TRUNCATED
+
+		SECTION("With Abort")
+		{
+			int Count = 0;
+
+			CLANG_BEGIN_DISABLE_OPTIMIZATIONS
+			AutoRTFM::ETransactionResult Result = AutoRTFM::Transact([&]
+			{
+				Count = swprintf(Buffer, 8, L"cat says '%ls'!", L"meow");
+				AutoRTFM::AbortTransaction();
+			});
+			CLANG_END_DISABLE_OPTIMIZATIONS
+
+			REQUIRE(AutoRTFM::ETransactionResult::AbortedByRequest == Result);
+			REQUIRE(0 == Count);
+			REQUIRE(L"_____________________" == std::wstring_view(Buffer));
+		}
+
+		SECTION("With Commit")
+		{
+			int Count = 0;
+
+			CLANG_BEGIN_DISABLE_OPTIMIZATIONS
+			AutoRTFM::Commit([&]
+			{
+				Count = swprintf(Buffer, 8, L"cat says '%ls'!", L"meow");
+			});
+			CLANG_END_DISABLE_OPTIMIZATIONS
+
+			REQUIRE(0 > Count);
+			REQUIRE(L"cat say" == std::wstring_view(Buffer));
+		}
+
+		CLANG_END_DISABLE_WARN_FORMAT_TRUNCATED
+	}
+
+	SECTION("Null buffer")
+	{
+		SECTION("With Abort")
+		{
+			int Count = 0;
+
+			CLANG_BEGIN_DISABLE_OPTIMIZATIONS
+			AutoRTFM::ETransactionResult Result = AutoRTFM::Transact([&]
+			{
+				Count = swprintf(nullptr, 0, L"cat says '%ls'!", L"meow");
+				AutoRTFM::AbortTransaction();
+			});
+			CLANG_END_DISABLE_OPTIMIZATIONS
+
+			REQUIRE(AutoRTFM::ETransactionResult::AbortedByRequest == Result);
+			REQUIRE(0 == Count);
+		}
+
+		SECTION("With Commit")
+		{
+			int Count = 0;
+
+			CLANG_BEGIN_DISABLE_OPTIMIZATIONS
+			AutoRTFM::Commit([&]
+			{
+				Count = swprintf(nullptr, 0, L"cat says '%ls'!", L"meow");
+			});
+			CLANG_END_DISABLE_OPTIMIZATIONS
+
+			REQUIRE(16 == Count);
+		}
+	}
+
+	SECTION("PercentN")
+	{
+		FCaptureWarningContext WarningContext;
+
+		int Count = 0;
+
+		CLANG_BEGIN_DISABLE_OPTIMIZATIONS
+		AutoRTFM::ETransactionResult Result = AutoRTFM::Transact([&]
+		{
+			swprintf(Buffer, 8, L"%n", &Count);
+		});
+		CLANG_END_DISABLE_OPTIMIZATIONS
+
+		REQUIRE(AutoRTFM::ETransactionResult::AbortedByLanguage == Result);
+		REQUIRE(0 == Count);
+		REQUIRE(L"_____________________" == std::wstring_view(Buffer));
+		REQUIRE_THAT(WarningContext.GetWarnings(),
+			Catch::Matchers::VectorContains(FString(kPercentNWarning)));
+	}
+
+	SECTION("PercentPercentN")
+	{
+		SECTION("With Abort")
+		{
+			int Count = 0;
+
+			CLANG_BEGIN_DISABLE_OPTIMIZATIONS
+			AutoRTFM::ETransactionResult Result = AutoRTFM::Transact([&]
+			{
+				Count = swprintf(Buffer, 8, L"%%n");
+				AutoRTFM::AbortTransaction();
+			});
+			CLANG_END_DISABLE_OPTIMIZATIONS
+
+			REQUIRE(AutoRTFM::ETransactionResult::AbortedByRequest == Result);
+			REQUIRE(0 == Count);
+			REQUIRE(L"_____________________" == std::wstring_view(Buffer));
+		}
+
+		SECTION("With Commit")
+		{
+			int Count = 0;
+
+			CLANG_BEGIN_DISABLE_OPTIMIZATIONS
+			AutoRTFM::Commit([&]
+			{
+				Count = swprintf(Buffer, 8, L"%%n");
+			});
+			CLANG_END_DISABLE_OPTIMIZATIONS
+
+			REQUIRE(2 == Count);
+			REQUIRE(L"%n" == std::wstring_view(Buffer));
+		}
+	}
+}
+
+TEST_CASE("CString.printf")
+{
+	SECTION("With Abort")
+	{
+		int Count = 0;
+
+		CLANG_BEGIN_DISABLE_OPTIMIZATIONS
+		AutoRTFM::ETransactionResult Result = AutoRTFM::Transact([&]
+		{
+			Count = printf("AutoRTFM printf test: %%n\n");
+			AutoRTFM::AbortTransaction();
+		});
+		CLANG_END_DISABLE_OPTIMIZATIONS
+
+		REQUIRE(AutoRTFM::ETransactionResult::AbortedByRequest == Result);
+		REQUIRE(0 == Count);
+	}
+
+	SECTION("With Commit")
+	{
+		int Count = 0;
+
+		CLANG_BEGIN_DISABLE_OPTIMIZATIONS
+		AutoRTFM::Commit([&]
+		{
+			Count = printf("AutoRTFM printf test: %%n\n");
+		});
+		CLANG_END_DISABLE_OPTIMIZATIONS
+
+		REQUIRE(25 == Count);
+	}
+
+	SECTION("PercentN")
+	{
+		FCaptureWarningContext WarningContext;
+
+		int Count = 0;
+
+		CLANG_BEGIN_DISABLE_OPTIMIZATIONS
+		AutoRTFM::ETransactionResult Result = AutoRTFM::Transact([&]
+		{
+			printf("%n", &Count);
+		});
+		CLANG_END_DISABLE_OPTIMIZATIONS
+
+		REQUIRE(AutoRTFM::ETransactionResult::AbortedByLanguage == Result);
+		REQUIRE(0 == Count);
+		REQUIRE_THAT(WarningContext.GetWarnings(),
+			Catch::Matchers::VectorContains(FString(kPercentNWarning)));
+	}
+}
+
+TEST_CASE("CString.wprintf")
+{
+	SECTION("With Abort")
+	{
+		int Count = 0;
+
+		CLANG_BEGIN_DISABLE_OPTIMIZATIONS
+		AutoRTFM::ETransactionResult Result = AutoRTFM::Transact([&]
+		{
+			Count = wprintf(L"AutoRTFM wprintf test: %%n\n");
+			AutoRTFM::AbortTransaction();
+		});
+		CLANG_END_DISABLE_OPTIMIZATIONS
+
+		REQUIRE(AutoRTFM::ETransactionResult::AbortedByRequest == Result);
+		REQUIRE(0 == Count);
+	}
+
+	SECTION("With Commit")
+	{
+		int Count = 0;
+
+		CLANG_BEGIN_DISABLE_OPTIMIZATIONS
+		AutoRTFM::Commit([&]
+		{
+			Count = wprintf(L"AutoRTFM wprintf test: %%n\n");
+		});
+		CLANG_END_DISABLE_OPTIMIZATIONS
+
+		REQUIRE(26 == Count);
+	}
+
+	SECTION("PercentN")
+	{
+		FCaptureWarningContext WarningContext;
+
+		int Count = 0;
+
+		CLANG_BEGIN_DISABLE_OPTIMIZATIONS
+		AutoRTFM::ETransactionResult Result = AutoRTFM::Transact([&]
+		{
+			wprintf(L"%n", &Count);
+		});
+		CLANG_END_DISABLE_OPTIMIZATIONS
+
+		REQUIRE(AutoRTFM::ETransactionResult::AbortedByLanguage == Result);
+		REQUIRE(0 == Count);
+		REQUIRE_THAT(WarningContext.GetWarnings(),
+			Catch::Matchers::VectorContains(FString(kPercentNWarning)));
+	}
+}
+
 #undef CLANG_BEGIN_DISABLE_OPTIMIZATIONS
 #undef CLANG_END_DISABLE_OPTIMIZATIONS
 #undef MSVC_BEGIN_DISABLE_WARN_UNSAFE_FUNCTION
diff --git a/Engine/Source/Runtime/Core/Private/AutoRTFM/StdLibSupport.cpp b/Engine/Source/Runtime/Core/Private/AutoRTFM/StdLibSupport.cpp
--- a/Engine/Source/Runtime/Core/Private/AutoRTFM/StdLibSupport.cpp
+++ b/Engine/Source/Runtime/Core/Private/AutoRTFM/StdLibSupport.cpp
@@ -23,6 +23,56 @@
 namespace AutoRTFM
 {
 
+namespace
+{
+
+void ThrowErrorFormatContainsPercentN()
+{
+    UE_LOG(LogAutoRTFM, Warning, TEXT("AutoRTFM does not support format strings containing '%%n'"));
+    FContext* Context = FContext::Get();
+    Context->AbortByLanguageAndThrow();
+}
+
+// Throws an error if the format string contains a '%n'.
+static void ThrowIfFormatContainsPercentN(const char* Format)
+{
+    for (const char* P = Format; *P != '\0'; ++P)
+    {
+        if (*P == '%')
+        {
+            switch (*++P)
+            {
+                case 'n':
+                    ThrowErrorFormatContainsPercentN();
+                    break;
+                case '\0':
+                    return;
+            }
+        }
+    }
+}
+
+// Throws an error if the format string contains a '%n'.
+static void ThrowIfFormatContainsPercentN(const wchar_t* Format)
+{
+    for (const wchar_t* P = Format; *P != L'\0'; ++P)
+    {
+        if (*P == L'%')
+        {
+            switch (*++P)
+            {
+                case L'n':
+                    ThrowErrorFormatContainsPercentN();
+                    break;
+                case L'\0':
+                    return;
+            }
+        }
+    }
+}
+
+}
+
 UE_AUTORTFM_REGISTER_OPEN_FUNCTION_EXPLICIT(memcpy, Memcpy);
 UE_AUTORTFM_REGISTER_OPEN_FUNCTION_EXPLICIT(memmove, Memmove);
 UE_AUTORTFM_REGISTER_OPEN_FUNCTION_EXPLICIT(memset, Memset);
@@ -93,7 +143,7 @@ char* RTFM_strcpy(char* const Dst, const char* const Src)
     const size_t SrcLen = strlen(Src);
 
 	FContext* Context = FContext::Get();
-    Context->RecordWrite(Dst, SrcLen);
+    Context->RecordWrite(Dst, SrcLen + sizeof(char));
     return strcpy(Dst, Src);
 }
 UE_AUTORTFM_REGISTER_OPEN_FUNCTION(strcpy);
@@ -134,7 +184,18 @@ UE_AUTORTFM_REGISTER_SELF_FUNCTION(static_cast<const char*(*)(const char*, int)>
 UE_AUTORTFM_REGISTER_SELF_FUNCTION(static_cast<const char* (*)(const char*, int)>(&strrchr));
 UE_AUTORTFM_REGISTER_SELF_FUNCTION(static_cast<const char* (*)(const char*, const char*)>(&strstr));
 UE_AUTORTFM_REGISTER_SELF_FUNCTION(strlen);
-UE_AUTORTFM_REGISTER_SELF_FUNCTION(strtol);
+
+long int RTFM_strtol(const char* String, char** EndPtr, int Radix)
+{
+    if (nullptr != EndPtr)
+    {
+        FContext* Context = FContext::Get();
+        Context->RecordWrite(EndPtr, sizeof(char*));
+    }
+
+    return strtol(String, EndPtr, Radix);
+}
+UE_AUTORTFM_REGISTER_OPEN_FUNCTION(strtol);
 
 UE_AUTORTFM_REGISTER_SELF_FUNCTION(static_cast<const wchar_t* (*)(const wchar_t*, wchar_t)>(&wcschr));
 UE_AUTORTFM_REGISTER_SELF_FUNCTION(static_cast<wchar_t* (*)(wchar_t*, wchar_t)>(&wcschr));
@@ -239,24 +300,72 @@ UE_AUTORTFM_REGISTER_SELF_FUNCTION(fmodf);
 UE_AUTORTFM_REGISTER_SELF_FUNCTION(fmodl);
 UE_AUTORTFM_REGISTER_SELF_FUNCTION(rand);
 
-// FIXME: This is only correct when:
-// - Str is newly allocated
-// - Format is either newly allocated or not mutated
-// - any strings passed as arguments are either newly allocated or not mutated
-int RTFM_snprintf(char* Str, size_t Size, char* Format, ...)
+// FIXME: Does not currently support %n format specifiers.
+int RTFM_vsnprintf(char* Str, size_t Size, const char* Format, va_list ArgList)
+{
+    ThrowIfFormatContainsPercentN(Format);
+
+    if (nullptr != Str && 0 != Size)
+    {
+        va_list ArgList2;
+        va_copy(ArgList2, ArgList);
+
+        FContext* Context = FContext::Get();
+        int Count = vsnprintf(nullptr, 0, Format, ArgList2);
+        if (Count >= 0)
+        {
+            size_t NumBytes = std::min<size_t>(Size, (1 + Count)) * sizeof(char);
+            Context->RecordWrite(Str, NumBytes);
+        }
+    }
+
+    return vsnprintf(Str, Size, Format, ArgList);
+}
+UE_AUTORTFM_REGISTER_OPEN_FUNCTION(vsnprintf);
+
+// FIXME: Does not currently support %n format specifiers.
+int RTFM_vswprintf(wchar_t* Str, size_t Size, const wchar_t* Format, va_list ArgList)
+{
+    ThrowIfFormatContainsPercentN(Format);
+
+    if (nullptr != Str && 0 != Size)
+    {
+        va_list ArgList2;
+        va_copy(ArgList2, ArgList);
+
+        FContext* Context = FContext::Get();
+        int Count = vswprintf(nullptr, 0, Format, ArgList2);
+        if (Count >= 0)
+        {
+            size_t NumBytes = std::min<size_t>(Size, (1 + Count)) * sizeof(wchar_t);
+            Context->RecordWrite(Str, NumBytes);
+        }
+    }
+
+    return vswprintf(Str, Size, Format, ArgList);
+}
+UE_AUTORTFM_REGISTER_OPEN_FUNCTION_EXPLICIT(
+    static_cast<int(*)(wchar_t*, size_t, const wchar_t*, va_list)>(&vswprintf),
+    RTFM_vswprintf);
+
+// FIXME: Does not currently support %n format specifiers.
+int RTFM_snprintf(char* Str, size_t Size, const char* Format, ...)
 {
-	FContext* Context = FContext::Get();
     va_list ArgList;
+
     va_start(ArgList, Format);
-    int Result = vsnprintf(Str, Size, Format, ArgList);
+    int Count = RTFM_vsnprintf(Str, Size, Format, ArgList);
     va_end(ArgList);
-    return Result;
+
+    return Count;
 }
 UE_AUTORTFM_REGISTER_OPEN_FUNCTION(snprintf);
 
+// FIXME: Does not currently support %n format specifiers.
 int RTFM_printf(const char* Format, ...)
 {
-	FContext* Context = FContext::Get();
+    ThrowIfFormatContainsPercentN(Format);
+
     va_list ArgList;
     va_start(ArgList, Format);
     int Result = vprintf(Format, ArgList);
@@ -265,6 +374,19 @@ int RTFM_printf(const char* Format, ...)
 }
 UE_AUTORTFM_REGISTER_OPEN_FUNCTION(printf);
 
+// FIXME: Does not currently support %n format specifiers.
+int RTFM_wprintf(const wchar_t* Format, ...)
+{
+    ThrowIfFormatContainsPercentN(Format);
+
+    va_list ArgList;
+    va_start(ArgList, Format);
+    int Result = vwprintf(Format, ArgList);
+    va_end(ArgList);
+    return Result;
+}
+UE_AUTORTFM_REGISTER_OPEN_FUNCTION(wprintf);
+
 int RTFM_putchar(int Char)
 {
     return putchar(Char);
@@ -296,10 +418,92 @@ FILE* RTFM___acrt_iob_func(int Index)
     }
 }
 UE_AUTORTFM_REGISTER_OPEN_FUNCTION(__acrt_iob_func);
-UE_AUTORTFM_REGISTER_SELF_FUNCTION(__stdio_common_vfprintf);
-UE_AUTORTFM_REGISTER_SELF_FUNCTION(__stdio_common_vsprintf);
-UE_AUTORTFM_REGISTER_SELF_FUNCTION(__stdio_common_vswprintf);
-UE_AUTORTFM_REGISTER_SELF_FUNCTION(__stdio_common_vfwprintf);
+
+// FIXME: Does not currently support %n format specifiers.
+int RTFM___stdio_common_vfprintf(
+        unsigned __int64 Options,
+        FILE*            Stream,
+        char const*      Format,
+        _locale_t        Locale,
+        va_list          ArgList)
+{
+    ThrowIfFormatContainsPercentN(Format);
+
+    return __stdio_common_vfprintf(Options, Stream, Format, Locale, ArgList);
+}
+UE_AUTORTFM_REGISTER_OPEN_FUNCTION(__stdio_common_vfprintf);
+
+// FIXME: Does not currently support %n format specifiers.
+int RTFM___stdio_common_vsprintf(
+        unsigned __int64 Options,
+        char*            Buffer,
+        size_t           BufferCount,
+        char const*      Format,
+        _locale_t        Locale,
+        va_list          ArgList)
+{
+    ThrowIfFormatContainsPercentN(Format);
+
+    if (nullptr != Buffer && 0 != BufferCount)
+    {
+        va_list ArgList2;
+        va_copy(ArgList2, ArgList);
+
+        FContext* Context = FContext::Get();
+        int Count = __stdio_common_vsprintf(Options, nullptr, 0, Format, Locale, ArgList2);
+        if (Count >= 0)
+        {
+            size_t NumBytes = std::min<size_t>(BufferCount, (1 + Count)) * sizeof(char);
+            Context->RecordWrite(Buffer, NumBytes);
+        }
+    }
+
+    return __stdio_common_vsprintf(Options, Buffer, BufferCount, Format, Locale, ArgList);
+}
+UE_AUTORTFM_REGISTER_OPEN_FUNCTION(__stdio_common_vsprintf);
+
+// FIXME: Does not currently support %n format specifiers.
+int RTFM___stdio_common_vswprintf(
+        unsigned __int64 Options,
+        wchar_t*         Buffer,
+        size_t           BufferCount,
+        wchar_t const*   Format,
+        _locale_t        Locale,
+        va_list          ArgList)
+{
+    ThrowIfFormatContainsPercentN(Format);
+
+    if (nullptr != Buffer && 0 != BufferCount)
+    {
+        va_list ArgList2;
+        va_copy(ArgList2, ArgList);
+
+        FContext* Context = FContext::Get();
+        int Count = __stdio_common_vswprintf(Options, nullptr, 0, Format, Locale, ArgList2);
+        if (Count >= 0)
+        {
+            size_t NumBytes = std::min<size_t>(BufferCount, (1 + Count)) * sizeof(wchar_t);
+            Context->RecordWrite(Buffer, NumBytes);
+        }
+    }
+
+    return __stdio_common_vswprintf(Options, Buffer, BufferCount, Format, Locale, ArgList);
+}
+UE_AUTORTFM_REGISTER_OPEN_FUNCTION(__stdio_common_vswprintf);
+
+// FIXME: Does not currently support %n format specifiers.
+int RTFM___stdio_common_vfwprintf(
+        unsigned __int64 Options,
+        FILE*            Stream,
+        wchar_t const*   Format,
+        _locale_t        Locale,
+        va_list          ArgList)
+{
+    ThrowIfFormatContainsPercentN(Format);
+
+    return __stdio_common_vfwprintf(Options, Stream, Format, Locale, ArgList);
+}
+UE_AUTORTFM_REGISTER_OPEN_FUNCTION(__stdio_common_vfwprintf);
 
 #else // PLATFORM_WINDOWS -> so !PLATFORM_WINDOWS
 extern "C" size_t _ZNSt3__112__next_primeEm(size_t N) __attribute__((weak));
@@ -351,6 +555,15 @@ UE_AUTORTFM_REGISTER_SELF_FUNCTION(bcmp);
 UE_AUTORTFM_REGISTER_SELF_FUNCTION(pthread_getspecific);
 #endif // PLATFORM_LINUX
 
+wchar_t* RTFM_wcscpy(wchar_t* Dst, const wchar_t* Src)
+{
+    const size_t SrcLen = wcslen(Src);
+
+	FContext* Context = FContext::Get();
+    Context->RecordWrite(Dst, (SrcLen + 1) * sizeof(wchar_t));
+    return wcscpy(Dst, Src);
+}
+
 wchar_t* RTFM_wcsncpy(wchar_t* Dst, const wchar_t* Src, size_t Count)
 {
 	FContext* Context = FContext::Get();
@@ -367,6 +580,7 @@ wchar_t* RTFM_wcsncpy(wchar_t* Dst, const wchar_t* Src, size_t Count)
 #pragma warning(push)
 #endif
 
+UE_AUTORTFM_REGISTER_OPEN_FUNCTION(wcscpy);
 UE_AUTORTFM_REGISTER_OPEN_FUNCTION(wcsncpy);
 
 #ifdef _MSC_VER
