Fix some fun TeleportTo found bugs with child actor components by adding a new transactionally safe RWLock, and also fixing a bug in the existing transactionally safe critical section code.

diff --git a/Engine/Source/Programs/AutoRTFMTests/Private/UECoreTests.cpp b/Engine/Source/Programs/AutoRTFMTests/Private/UECoreTests.cpp
--- a/Engine/Source/Programs/AutoRTFMTests/Private/UECoreTests.cpp
+++ b/Engine/Source/Programs/AutoRTFMTests/Private/UECoreTests.cpp
@@ -10,6 +10,8 @@
 #include "UObject/NameTypes.h"
 #include "UObject/UObjectArray.h"
 #include "MyAutoRTFMTestObject.h"
+#include "Misc/TransactionallySafeScopeLock.h"
+#include "Misc/TransactionallySafeRWScopeLock.h"
 
 TEST_CASE("UECore.FDelegateHandle")
 {
@@ -962,3 +964,510 @@ TEST_CASE("UECore.FUObjectItem")
 		}
 	}
 }
+
+TEST_CASE("UECore.TransactionallySafeScopeLock")
+{
+	SECTION("Outside Transaction")
+	{
+		FTransactionallySafeCriticalSection CriticalSection;
+
+		AutoRTFM::ETransactionResult Result = AutoRTFM::Transact([&]
+			{
+				FTransactionallySafeScopeLock Lock(&CriticalSection);
+				AutoRTFM::AbortTransaction();
+			});
+
+		REQUIRE(AutoRTFM::ETransactionResult::AbortedByRequest == Result);
+
+		Result = AutoRTFM::Transact([&]
+			{
+				FTransactionallySafeScopeLock Lock(&CriticalSection);
+			});
+
+		REQUIRE(AutoRTFM::ETransactionResult::Committed == Result);
+	}
+
+	SECTION("Inside Transaction")
+	{
+		AutoRTFM::ETransactionResult Result = AutoRTFM::Transact([&]
+			{
+				FTransactionallySafeCriticalSection CriticalSection;
+				FTransactionallySafeScopeLock Lock(&CriticalSection);
+				AutoRTFM::AbortTransaction();
+			});
+
+		REQUIRE(AutoRTFM::ETransactionResult::AbortedByRequest == Result);
+
+		Result = AutoRTFM::Transact([&]
+			{
+				FTransactionallySafeCriticalSection CriticalSection;
+				FTransactionallySafeScopeLock Lock(&CriticalSection);
+			});
+
+		REQUIRE(AutoRTFM::ETransactionResult::Committed == Result);
+	}
+
+	SECTION("Inside Transaction Used In Nested Transaction")
+	{
+		AutoRTFM::ETransactionResult InnerResult = AutoRTFM::ETransactionResult::Committed;
+
+		AutoRTFM::ETransactionResult Result = AutoRTFM::Transact([&]
+			{
+				FTransactionallySafeCriticalSection CriticalSection;
+
+				InnerResult = AutoRTFM::Transact([&]
+					{
+						FTransactionallySafeScopeLock Lock(&CriticalSection);
+						AutoRTFM::CascadingAbortTransaction();
+					});
+			});
+
+		REQUIRE(AutoRTFM::ETransactionResult::AbortedByCascade == Result);
+
+		Result = AutoRTFM::Transact([&]
+			{
+				FTransactionallySafeCriticalSection CriticalSection;
+
+				InnerResult = AutoRTFM::Transact([&]
+					{
+						FTransactionallySafeScopeLock Lock(&CriticalSection);
+						AutoRTFM::AbortTransaction();
+					});
+			});
+
+		REQUIRE(AutoRTFM::ETransactionResult::AbortedByRequest == InnerResult);
+		REQUIRE(AutoRTFM::ETransactionResult::Committed == Result);
+
+		Result = AutoRTFM::Transact([&]
+			{
+				FTransactionallySafeCriticalSection CriticalSection;
+
+				InnerResult = AutoRTFM::Transact([&]
+					{
+						FTransactionallySafeScopeLock Lock(&CriticalSection);
+					});
+
+				AutoRTFM::AbortTransaction();
+			});
+
+		REQUIRE(AutoRTFM::ETransactionResult::AbortedByRequest == Result);
+
+		Result = AutoRTFM::Transact([&]
+			{
+				FTransactionallySafeCriticalSection CriticalSection;
+
+				InnerResult = AutoRTFM::Transact([&]
+					{
+						FTransactionallySafeScopeLock Lock(&CriticalSection);
+					});
+			});
+
+		REQUIRE(AutoRTFM::ETransactionResult::Committed == InnerResult);
+		REQUIRE(AutoRTFM::ETransactionResult::Committed == Result);
+	}
+}
+
+TEST_CASE("UECore.TransactionallySafeRWScopeLock")
+{
+	SECTION("Outside Transaction With Read Lock")
+	{
+		FTransactionallySafeRWLock RWLock;
+
+		AutoRTFM::ETransactionResult Result = AutoRTFM::Transact([&]
+			{
+				FTransactionallySafeRWScopeLock Lock(RWLock, FRWScopeLockType::SLT_ReadOnly);
+				AutoRTFM::AbortTransaction();
+			});
+
+		REQUIRE(AutoRTFM::ETransactionResult::AbortedByRequest == Result);
+
+		Result = AutoRTFM::Transact([&]
+			{
+				FTransactionallySafeRWScopeLock Lock(RWLock, FRWScopeLockType::SLT_ReadOnly);
+			});
+
+		REQUIRE(AutoRTFM::ETransactionResult::Committed == Result);
+	}
+
+	SECTION("Inside Transaction With Read Lock")
+	{
+		AutoRTFM::ETransactionResult Result = AutoRTFM::Transact([&]
+			{
+				FTransactionallySafeRWLock RWLock;
+				FTransactionallySafeRWScopeLock Lock(RWLock, FRWScopeLockType::SLT_ReadOnly);
+				AutoRTFM::AbortTransaction();
+			});
+
+		REQUIRE(AutoRTFM::ETransactionResult::AbortedByRequest == Result);
+
+		Result = AutoRTFM::Transact([&]
+			{
+				FTransactionallySafeRWLock RWLock;
+				FTransactionallySafeRWScopeLock Lock(RWLock, FRWScopeLockType::SLT_ReadOnly);
+			});
+
+		REQUIRE(AutoRTFM::ETransactionResult::Committed == Result);
+	}
+
+	SECTION("Inside Transaction Used In Nested Transaction With Read Lock")
+	{
+		AutoRTFM::ETransactionResult InnerResult = AutoRTFM::ETransactionResult::Committed;
+
+		AutoRTFM::ETransactionResult Result = AutoRTFM::Transact([&]
+			{
+				FTransactionallySafeRWLock RWLock;
+
+				InnerResult = AutoRTFM::Transact([&]
+					{
+						FTransactionallySafeRWScopeLock Lock(RWLock, FRWScopeLockType::SLT_ReadOnly);
+						AutoRTFM::CascadingAbortTransaction();
+					});
+			});
+
+		REQUIRE(AutoRTFM::ETransactionResult::AbortedByCascade == Result);
+
+		Result = AutoRTFM::Transact([&]
+			{
+				FTransactionallySafeRWLock RWLock;
+
+				InnerResult = AutoRTFM::Transact([&]
+					{
+						FTransactionallySafeRWScopeLock Lock(RWLock, FRWScopeLockType::SLT_ReadOnly);
+						AutoRTFM::AbortTransaction();
+					});
+			});
+
+		REQUIRE(AutoRTFM::ETransactionResult::AbortedByRequest == InnerResult);
+		REQUIRE(AutoRTFM::ETransactionResult::Committed == Result);
+
+		Result = AutoRTFM::Transact([&]
+			{
+				FTransactionallySafeRWLock RWLock;
+
+				InnerResult = AutoRTFM::Transact([&]
+					{
+						FTransactionallySafeRWScopeLock Lock(RWLock, FRWScopeLockType::SLT_ReadOnly);
+					});
+
+				AutoRTFM::AbortTransaction();
+			});
+
+		REQUIRE(AutoRTFM::ETransactionResult::AbortedByRequest == Result);
+
+		Result = AutoRTFM::Transact([&]
+			{
+				FTransactionallySafeRWLock RWLock;
+
+				InnerResult = AutoRTFM::Transact([&]
+					{
+						FTransactionallySafeRWScopeLock Lock(RWLock, FRWScopeLockType::SLT_ReadOnly);
+					});
+			});
+
+		REQUIRE(AutoRTFM::ETransactionResult::Committed == InnerResult);
+		REQUIRE(AutoRTFM::ETransactionResult::Committed == Result);
+	}
+
+	SECTION("Outside Transaction With Write Lock")
+	{
+		FTransactionallySafeRWLock RWLock;
+
+		AutoRTFM::ETransactionResult Result = AutoRTFM::Transact([&]
+			{
+				FTransactionallySafeRWScopeLock Lock(RWLock, FRWScopeLockType::SLT_Write);
+				AutoRTFM::AbortTransaction();
+			});
+
+		REQUIRE(AutoRTFM::ETransactionResult::AbortedByRequest == Result);
+
+		Result = AutoRTFM::Transact([&]
+			{
+				FTransactionallySafeRWScopeLock Lock(RWLock, FRWScopeLockType::SLT_Write);
+			});
+
+		REQUIRE(AutoRTFM::ETransactionResult::Committed == Result);
+	}
+
+	SECTION("Inside Transaction With Write Lock")
+	{
+		AutoRTFM::ETransactionResult Result = AutoRTFM::Transact([&]
+			{
+				FTransactionallySafeRWLock RWLock;
+				FTransactionallySafeRWScopeLock Lock(RWLock, FRWScopeLockType::SLT_Write);
+				AutoRTFM::AbortTransaction();
+			});
+
+		REQUIRE(AutoRTFM::ETransactionResult::AbortedByRequest == Result);
+
+		Result = AutoRTFM::Transact([&]
+			{
+				FTransactionallySafeRWLock RWLock;
+				FTransactionallySafeRWScopeLock Lock(RWLock, FRWScopeLockType::SLT_Write);
+			});
+
+		REQUIRE(AutoRTFM::ETransactionResult::Committed == Result);
+	}
+
+	SECTION("Inside Transaction Used In Nested Transaction With Write Lock")
+	{
+		AutoRTFM::ETransactionResult InnerResult = AutoRTFM::ETransactionResult::Committed;
+
+		AutoRTFM::ETransactionResult Result = AutoRTFM::Transact([&]
+			{
+				FTransactionallySafeRWLock RWLock;
+
+				InnerResult = AutoRTFM::Transact([&]
+					{
+						FTransactionallySafeRWScopeLock Lock(RWLock, FRWScopeLockType::SLT_Write);
+						AutoRTFM::CascadingAbortTransaction();
+					});
+			});
+
+		REQUIRE(AutoRTFM::ETransactionResult::AbortedByCascade == Result);
+
+		Result = AutoRTFM::Transact([&]
+			{
+				FTransactionallySafeRWLock RWLock;
+
+				InnerResult = AutoRTFM::Transact([&]
+					{
+						FTransactionallySafeRWScopeLock Lock(RWLock, FRWScopeLockType::SLT_Write);
+						AutoRTFM::AbortTransaction();
+					});
+			});
+
+		REQUIRE(AutoRTFM::ETransactionResult::AbortedByRequest == InnerResult);
+		REQUIRE(AutoRTFM::ETransactionResult::Committed == Result);
+
+		Result = AutoRTFM::Transact([&]
+			{
+				FTransactionallySafeRWLock RWLock;
+
+				InnerResult = AutoRTFM::Transact([&]
+					{
+						FTransactionallySafeRWScopeLock Lock(RWLock, FRWScopeLockType::SLT_Write);
+					});
+
+				AutoRTFM::AbortTransaction();
+			});
+
+		REQUIRE(AutoRTFM::ETransactionResult::AbortedByRequest == Result);
+
+		Result = AutoRTFM::Transact([&]
+			{
+				FTransactionallySafeRWLock RWLock;
+
+				InnerResult = AutoRTFM::Transact([&]
+					{
+						FTransactionallySafeRWScopeLock Lock(RWLock, FRWScopeLockType::SLT_Write);
+					});
+			});
+
+		REQUIRE(AutoRTFM::ETransactionResult::Committed == InnerResult);
+		REQUIRE(AutoRTFM::ETransactionResult::Committed == Result);
+	}
+}
+
+TEST_CASE("UECore.TransactionallySafeReadScopeLock")
+{
+	SECTION("Outside Transaction")
+	{
+		FTransactionallySafeRWLock RWLock;
+
+		AutoRTFM::ETransactionResult Result = AutoRTFM::Transact([&]
+			{
+				FTransactionallySafeReadScopeLock Lock(RWLock);
+				AutoRTFM::AbortTransaction();
+			});
+
+		REQUIRE(AutoRTFM::ETransactionResult::AbortedByRequest == Result);
+
+		Result = AutoRTFM::Transact([&]
+			{
+				FTransactionallySafeReadScopeLock Lock(RWLock);
+			});
+
+		REQUIRE(AutoRTFM::ETransactionResult::Committed == Result);
+	}
+
+	SECTION("Inside Transaction")
+	{
+		AutoRTFM::ETransactionResult Result = AutoRTFM::Transact([&]
+			{
+				FTransactionallySafeRWLock RWLock;
+				FTransactionallySafeReadScopeLock Lock(RWLock);
+				AutoRTFM::AbortTransaction();
+			});
+
+		REQUIRE(AutoRTFM::ETransactionResult::AbortedByRequest == Result);
+
+		Result = AutoRTFM::Transact([&]
+			{
+				FTransactionallySafeRWLock RWLock;
+				FTransactionallySafeReadScopeLock Lock(RWLock);
+			});
+
+		REQUIRE(AutoRTFM::ETransactionResult::Committed == Result);
+	}
+
+	SECTION("Inside Transaction Used In Nested Transaction")
+	{
+		AutoRTFM::ETransactionResult InnerResult = AutoRTFM::ETransactionResult::Committed;
+
+		AutoRTFM::ETransactionResult Result = AutoRTFM::Transact([&]
+			{
+				FTransactionallySafeRWLock RWLock;
+
+				InnerResult = AutoRTFM::Transact([&]
+					{
+						FTransactionallySafeReadScopeLock Lock(RWLock);
+						AutoRTFM::CascadingAbortTransaction();
+					});
+			});
+
+		REQUIRE(AutoRTFM::ETransactionResult::AbortedByCascade == Result);
+
+		Result = AutoRTFM::Transact([&]
+			{
+				FTransactionallySafeRWLock RWLock;
+
+				InnerResult = AutoRTFM::Transact([&]
+					{
+						FTransactionallySafeReadScopeLock Lock(RWLock);
+						AutoRTFM::AbortTransaction();
+					});
+			});
+
+		REQUIRE(AutoRTFM::ETransactionResult::AbortedByRequest == InnerResult);
+		REQUIRE(AutoRTFM::ETransactionResult::Committed == Result);
+
+		Result = AutoRTFM::Transact([&]
+			{
+				FTransactionallySafeRWLock RWLock;
+
+				InnerResult = AutoRTFM::Transact([&]
+					{
+						FTransactionallySafeReadScopeLock Lock(RWLock);
+					});
+
+				AutoRTFM::AbortTransaction();
+			});
+
+		REQUIRE(AutoRTFM::ETransactionResult::AbortedByRequest == Result);
+
+		Result = AutoRTFM::Transact([&]
+			{
+				FTransactionallySafeRWLock RWLock;
+
+				InnerResult = AutoRTFM::Transact([&]
+					{
+						FTransactionallySafeReadScopeLock Lock(RWLock);
+					});
+			});
+
+		REQUIRE(AutoRTFM::ETransactionResult::Committed == InnerResult);
+		REQUIRE(AutoRTFM::ETransactionResult::Committed == Result);
+	}
+}
+
+TEST_CASE("UECore.TransactionallySafeWriteScopeLock")
+{
+	SECTION("Outside Transaction")
+	{
+		FTransactionallySafeRWLock RWLock;
+
+		AutoRTFM::ETransactionResult Result = AutoRTFM::Transact([&]
+			{
+				FTransactionallySafeWriteScopeLock Lock(RWLock);
+				AutoRTFM::AbortTransaction();
+			});
+
+		REQUIRE(AutoRTFM::ETransactionResult::AbortedByRequest == Result);
+
+		Result = AutoRTFM::Transact([&]
+			{
+				FTransactionallySafeWriteScopeLock Lock(RWLock);
+			});
+
+		REQUIRE(AutoRTFM::ETransactionResult::Committed == Result);
+	}
+
+	SECTION("Inside Transaction")
+	{
+		AutoRTFM::ETransactionResult Result = AutoRTFM::Transact([&]
+			{
+				FTransactionallySafeRWLock RWLock;
+				FTransactionallySafeWriteScopeLock Lock(RWLock);
+				AutoRTFM::AbortTransaction();
+			});
+
+		REQUIRE(AutoRTFM::ETransactionResult::AbortedByRequest == Result);
+
+		Result = AutoRTFM::Transact([&]
+			{
+				FTransactionallySafeRWLock RWLock;
+				FTransactionallySafeWriteScopeLock Lock(RWLock);
+			});
+
+		REQUIRE(AutoRTFM::ETransactionResult::Committed == Result);
+	}
+
+	SECTION("Inside Transaction Used In Nested Transaction")
+	{
+		AutoRTFM::ETransactionResult InnerResult = AutoRTFM::ETransactionResult::Committed;
+
+		AutoRTFM::ETransactionResult Result = AutoRTFM::Transact([&]
+			{
+				FTransactionallySafeRWLock RWLock;
+
+				InnerResult = AutoRTFM::Transact([&]
+					{
+						FTransactionallySafeWriteScopeLock Lock(RWLock);
+						AutoRTFM::CascadingAbortTransaction();
+					});
+			});
+
+		REQUIRE(AutoRTFM::ETransactionResult::AbortedByCascade == Result);
+
+		Result = AutoRTFM::Transact([&]
+			{
+				FTransactionallySafeRWLock RWLock;
+
+				InnerResult = AutoRTFM::Transact([&]
+					{
+						FTransactionallySafeWriteScopeLock Lock(RWLock);
+						AutoRTFM::AbortTransaction();
+					});
+			});
+
+		REQUIRE(AutoRTFM::ETransactionResult::AbortedByRequest == InnerResult);
+		REQUIRE(AutoRTFM::ETransactionResult::Committed == Result);
+
+		Result = AutoRTFM::Transact([&]
+			{
+				FTransactionallySafeRWLock RWLock;
+
+				InnerResult = AutoRTFM::Transact([&]
+					{
+						FTransactionallySafeWriteScopeLock Lock(RWLock);
+					});
+
+				AutoRTFM::AbortTransaction();
+			});
+
+		REQUIRE(AutoRTFM::ETransactionResult::AbortedByRequest == Result);
+
+		Result = AutoRTFM::Transact([&]
+			{
+				FTransactionallySafeRWLock RWLock;
+
+				InnerResult = AutoRTFM::Transact([&]
+					{
+						FTransactionallySafeWriteScopeLock Lock(RWLock);
+					});
+			});
+
+		REQUIRE(AutoRTFM::ETransactionResult::Committed == InnerResult);
+		REQUIRE(AutoRTFM::ETransactionResult::Committed == Result);
+	}
+}
diff --git a/Engine/Source/Runtime/Core/Private/AutoRTFM/Context.h b/Engine/Source/Runtime/Core/Private/AutoRTFM/Context.h
--- a/Engine/Source/Runtime/Core/Private/AutoRTFM/Context.h
+++ b/Engine/Source/Runtime/Core/Private/AutoRTFM/Context.h
@@ -50,8 +50,8 @@ public:
 
 	inline FTransaction* GetCurrentTransaction() const { return CurrentTransaction; }
 	inline FCallNest* GetCurrentNest() const { return CurrentNest; }
-    inline bool IsTransactionStack(void* LogicalAddress) const { return LogicalAddress >= StackBegin && LogicalAddress < OuterTransactStackAddress; }
-	inline bool IsInnerTransactionStack(void* LogicalAddress) const { return LogicalAddress >= StackBegin && LogicalAddress < CurrentTransactStackAddress; }
+    inline bool IsTransactionStack(const void* LogicalAddress) const { return LogicalAddress >= StackBegin && LogicalAddress < OuterTransactStackAddress; }
+	inline bool IsInnerTransactionStack(const void* LogicalAddress) const { return LogicalAddress >= StackBegin && LogicalAddress < CurrentTransactStackAddress; }
 	inline EContextStatus GetStatus() const { return Status; }
 	void Throw();
 	
diff --git a/Engine/Source/Runtime/Core/Public/Misc/TransactionallySafeCriticalSection.h b/Engine/Source/Runtime/Core/Public/Misc/TransactionallySafeCriticalSection.h
--- a/Engine/Source/Runtime/Core/Public/Misc/TransactionallySafeCriticalSection.h
+++ b/Engine/Source/Runtime/Core/Public/Misc/TransactionallySafeCriticalSection.h
@@ -4,6 +4,7 @@
 
 #include "HAL/CriticalSection.h"
 #include "AutoRTFM/AutoRTFM.h"
+#include "Templates/SharedPointer.h"
 
 // A transactionally safe critical section that works in the following novel ways:
 // - In the open (non-transactional):
@@ -24,37 +25,78 @@
 // modifications we should make.
 struct FTransactionallySafeCriticalSectionDefinition final
 {
+	// Always open because the constructor arguments will create the underlying critical section.
+	UE_AUTORTFM_ALWAYS_OPEN
+	FTransactionallySafeCriticalSectionDefinition() : State(new FState())
+	{
+		if (AutoRTFM::IsTransactional())
+		{
+			const AutoRTFM::EContextStatus Status = AutoRTFM::Close([this]
+				{
+					// We explicitly copy the state here for the case that `this` was stack
+					// allocated and has already died before the on-abort is hit.
+					AutoRTFM::OnAbort([State = this->State]
+						{
+							ensure(0 == State->TransactionalLockCount); 
+						});
+				});
+
+			ensure(AutoRTFM::EContextStatus::OnTrack == Status);
+		}
+	}
+
+	~FTransactionallySafeCriticalSectionDefinition()
+	{
+		if (AutoRTFM::IsTransactional())
+		{
+			const AutoRTFM::EContextStatus Status = AutoRTFM::Close([this]
+				{
+					// We explicitly copy the state here for the case that `this` was stack
+					// allocated and has already died before the on-commit is hit.
+					AutoRTFM::OnCommit([State = this->State]
+						{
+							ensure(0 == State->TransactionalLockCount);
+						});
+				});
+
+			ensure(AutoRTFM::EContextStatus::OnTrack == Status);
+		}
+	}
+
 	void Lock()
 	{
 		if (AutoRTFM::IsTransactional() || AutoRTFM::IsCommittingOrAborting())
 		{
-			AutoRTFM::Open([&]
+			AutoRTFM::Open([this]
 				{
 					// The transactional system which can increment TransactionalLockCount
 					// is always single-threaded, thus this is safe to check without atomicity.
-					if (0 == TransactionalLockCount)
+					if (0 == State->TransactionalLockCount)
 					{
-						CriticalSection.Lock();
+						State->CriticalSection.Lock();
 					}
 
-					TransactionalLockCount += 1;
+					State->TransactionalLockCount += 1;
 				});
 
-			AutoRTFM::OnAbort([this]
+			// We explicitly copy the state here for the case that `this` was stack
+			// allocated and has already died before the on-abort is hit.
+			AutoRTFM::OnAbort([State = this->State]
 				{
-					ensure(0 != TransactionalLockCount);
-					TransactionalLockCount -= 1;
+					ensure(0 != State->TransactionalLockCount);
 
-					if (0 == TransactionalLockCount)
+					State->TransactionalLockCount -= 1;
+
+					if (0 == State->TransactionalLockCount)
 					{
-						CriticalSection.Unlock();
+						State->CriticalSection.Unlock();
 					}
 				});
 		}
 		else
 		{
-			CriticalSection.Lock();
-			ensure(0 == TransactionalLockCount);
+			State->CriticalSection.Lock();
+			ensure(0 == State->TransactionalLockCount);
 		}
 	}
 
@@ -62,27 +104,37 @@ struct FTransactionallySafeCriticalSectionDefinition final
 	{
 		if (AutoRTFM::IsTransactional() || AutoRTFM::IsCommittingOrAborting())
 		{
-			AutoRTFM::OnCommit([this]
+			// We explicitly copy the state here for the case that `this` was stack
+			// allocated and has already died before the on-commit is hit.
+			AutoRTFM::OnCommit([State = this->State]
 				{
-					ensure(0 != TransactionalLockCount);
-					TransactionalLockCount -= 1;
+					ensure(0 != State->TransactionalLockCount);
+
+					State->TransactionalLockCount -= 1;
 
-					if (0 == TransactionalLockCount)
+					if (0 == State->TransactionalLockCount)
 					{
-						CriticalSection.Unlock();
+						State->CriticalSection.Unlock();
 					}
 				});
 		}
 		else
 		{
-			ensure(0 == TransactionalLockCount);
-			CriticalSection.Unlock();
+			ensure(0 == State->TransactionalLockCount);
+			State->CriticalSection.Unlock();
 		}
 	}
 
 private:
-	FCriticalSection CriticalSection;
-	uint32 TransactionalLockCount = 0;
+	UE_NONCOPYABLE(FTransactionallySafeCriticalSectionDefinition)
+
+	struct FState final
+	{
+		FCriticalSection CriticalSection;
+		uint32 TransactionalLockCount = 0;
+	};
+
+	const TSharedPtr<FState> State;
 };
 
 #if UE_AUTORTFM
diff --git a/Engine/Source/Runtime/Core/Public/Misc/TransactionallySafeRWLock.h b/Engine/Source/Runtime/Core/Public/Misc/TransactionallySafeRWLock.h
new file mode 100644
--- /dev/null
+++ b/Engine/Source/Runtime/Core/Public/Misc/TransactionallySafeRWLock.h
@@ -0,0 +1,182 @@
+// Copyright Epic Games, Inc. All Rights Reserved.
+
+#pragma once
+
+#include "HAL/CriticalSection.h"
+#include "AutoRTFM/AutoRTFM.h"
+#include "Templates/SharedPointer.h"
+
+// A transactionally safe lock that works in the following novel ways:
+	// - In the open (non-transactional):
+	//   - Take the lock like before. Simple!
+	//   - Free the lock like before too.
+	// - In the closed (transactional):
+	//   - During locking we query `TransactionalLockCount`:
+	//	   - 0 means we haven't taken the lock within our transaction nest and need to acquire the lock.
+	//     - Otherwise we already have the lock (and are preventing non-transactional code seeing any
+	//       modifications we've made while holding the lock), so just bump `TransactionalLockCount`.
+	//     - We also register an on-abort handler to release the lock should we abort (but we need to
+	//       query `TransactionalLockCount` even there because we could be aborting an inner transaction
+	//       and the parent transaction still wants to have the lock held!).
+	//   - During unlocking we defer doing the unlock until the transaction commits.
+	//
+	// Thus with this approach we will hold this lock for the *entirety* of the transactional nest should
+	// we take the lock during the transaction, thus preventing non-transactional code from seeing any
+	// modifications we should make.
+	//
+	// If we are within a transaction, we pessimise our read-lock to a write-lock. Note: that it should
+	// potentially be possible to have read-locks work correctly, but serious care will have to be taken to
+	// ensure that we don't have:
+	//   Open Thread     Closed Thread
+	//   -----------     ReadLock
+	//   -----------     ReadUnlock
+	//   WriteLock       -------------
+	//   WriteUnlock     -------------
+	//   -----------     ReadLock      <- Invalid because the transaction can potentially observe side
+	//                                    effects of the open-threads writes!
+struct FTransactionallySafeRWLockDefinition final
+{
+	// Always open because the constructor arguments will create the underlying critical section.
+	UE_AUTORTFM_ALWAYS_OPEN
+		FTransactionallySafeRWLockDefinition() : State(new FState())
+	{
+		if (AutoRTFM::IsTransactional())
+		{
+			const AutoRTFM::EContextStatus Status = AutoRTFM::Close([this]
+				{
+					// We explicitly copy the state here for the case that `this` was stack
+					// allocated and has already died before the on-abort is hit.
+					AutoRTFM::OnAbort([State = this->State]
+						{
+						});
+				});
+
+			ensure(AutoRTFM::EContextStatus::OnTrack == Status);
+		}
+	}
+
+	~FTransactionallySafeRWLockDefinition()
+	{
+		if (AutoRTFM::IsTransactional())
+		{
+			const AutoRTFM::EContextStatus Status = AutoRTFM::Close([this]
+				{
+					// We explicitly copy the state here for the case that `this` was stack
+					// allocated and has already died before the on-commit is hit.
+					AutoRTFM::OnCommit([State = this->State]
+						{
+							ensure(0 == State->TransactionalLockCount);
+						});
+				});
+
+			ensure(AutoRTFM::EContextStatus::OnTrack == Status);
+		}
+	}
+
+	void ReadLock()
+	{
+		if (AutoRTFM::IsTransactional() || AutoRTFM::IsCommittingOrAborting())
+		{
+			// Transactionally pessimise ReadLock -> WriteLock.
+			WriteLock();
+		}
+		else
+		{
+			State->Lock.ReadLock();
+			ensure(0 == State->TransactionalLockCount);
+		}
+	}
+
+	void ReadUnlock()
+	{
+		if (AutoRTFM::IsTransactional() || AutoRTFM::IsCommittingOrAborting())
+		{
+			// Transactionally pessimise ReadUnlock -> WriteUnlock.
+			WriteUnlock();
+		}
+		else
+		{
+			ensure(0 == State->TransactionalLockCount);
+			State->Lock.ReadUnlock();
+		}
+	}
+
+	void WriteLock()
+	{
+		if (AutoRTFM::IsTransactional() || AutoRTFM::IsCommittingOrAborting())
+		{
+			AutoRTFM::Open([&]
+				{
+					// The transactional system which can increment TransactionalLockCount
+					// is always single-threaded, thus this is safe to check without atomicity.
+					if (0 == State->TransactionalLockCount)
+					{
+						State->Lock.WriteLock();
+					}
+
+					State->TransactionalLockCount += 1;
+				});
+
+			// We explicitly copy the state here for the case that `this` was stack
+			// allocated and has already died before the on-abort is hit.
+			AutoRTFM::OnAbort([State = this->State]
+				{
+					ensure(0 != State->TransactionalLockCount);
+
+					State->TransactionalLockCount -= 1;
+
+					if (0 == State->TransactionalLockCount)
+					{
+						State->Lock.WriteUnlock();
+					}
+				});
+		}
+		else
+		{
+			State->Lock.WriteLock();
+			ensure(0 == State->TransactionalLockCount);
+		}
+	}
+
+	void WriteUnlock()
+	{
+		if (AutoRTFM::IsTransactional() || AutoRTFM::IsCommittingOrAborting())
+		{
+			// We explicitly copy the state here for the case that `this` was stack
+			// allocated and has already died before the on-commit is hit.
+			AutoRTFM::OnCommit([State = this->State]
+				{
+					ensure(0 != State->TransactionalLockCount);
+
+					State->TransactionalLockCount -= 1;
+
+					if (0 == State->TransactionalLockCount)
+					{
+						State->Lock.WriteUnlock();
+					}
+				});
+		}
+		else
+		{
+			ensure(0 == State->TransactionalLockCount);
+			State->Lock.WriteUnlock();
+		}
+	}
+
+private:
+	UE_NONCOPYABLE(FTransactionallySafeRWLockDefinition)
+
+	struct FState final
+	{
+		FRWLock Lock;
+		uint32 TransactionalLockCount = 0;
+	};
+
+	const TSharedPtr<FState> State;
+};
+
+#if UE_AUTORTFM
+using FTransactionallySafeRWLock = FTransactionallySafeRWLockDefinition;
+#else
+using FTransactionallySafeRWLock = FRWLock;
+#endif
diff --git a/Engine/Source/Runtime/Core/Public/Misc/TransactionallySafeRWScopeLock.h b/Engine/Source/Runtime/Core/Public/Misc/TransactionallySafeRWScopeLock.h
new file mode 100644
--- /dev/null
+++ b/Engine/Source/Runtime/Core/Public/Misc/TransactionallySafeRWScopeLock.h
@@ -0,0 +1,83 @@
+// Copyright Epic Games, Inc. All Rights Reserved.
+
+#pragma once
+
+#include "TransactionallySafeRWLock.h"
+#include "ScopeRWLock.h"
+
+// A transactionally safe read scope lock - that uses a transactionally safe read-write lock to back it.
+class FTransactionallySafeReadScopeLock
+{
+public:
+	UE_NODISCARD_CTOR explicit FTransactionallySafeReadScopeLock(FTransactionallySafeRWLock& InLock)
+		: Lock(InLock)
+	{
+		Lock.ReadLock();
+	}
+
+	~FTransactionallySafeReadScopeLock()
+	{
+		Lock.ReadUnlock();
+	}
+
+private:
+	UE_NONCOPYABLE(FTransactionallySafeReadScopeLock);
+
+	FTransactionallySafeRWLock& Lock;
+};
+
+// A transactionally safe write scope lock - that uses a transactionally safe read-write lock to back it.
+class FTransactionallySafeWriteScopeLock
+{
+public:
+	UE_NODISCARD_CTOR explicit FTransactionallySafeWriteScopeLock(FTransactionallySafeRWLock& InLock)
+		: Lock(InLock)
+	{
+		Lock.WriteLock();
+	}
+
+	~FTransactionallySafeWriteScopeLock()
+	{
+		Lock.WriteUnlock();
+	}
+
+private:
+	UE_NONCOPYABLE(FTransactionallySafeWriteScopeLock);
+
+	FTransactionallySafeRWLock& Lock;
+};
+
+// A transactionally safe read-write scope lock - that uses a transactionally safe read-write lock to back it.
+struct FTransactionallySafeRWScopeLock final
+{
+	UE_NODISCARD_CTOR explicit FTransactionallySafeRWScopeLock(FTransactionallySafeRWLock& InLockObject, FRWScopeLockType InLockType)
+		: LockObject(InLockObject), LockType(InLockType)
+	{
+		if (LockType != SLT_ReadOnly)
+		{
+			LockObject.WriteLock();
+		}
+		else
+		{
+			LockObject.ReadLock();
+		}
+	}
+
+	~FTransactionallySafeRWScopeLock()
+	{
+		if (LockType == SLT_ReadOnly)
+		{
+			LockObject.ReadUnlock();
+		}
+		else
+		{
+			LockObject.WriteUnlock();
+		}
+	}
+
+private:
+	UE_NONCOPYABLE(FTransactionallySafeRWScopeLock);
+
+	FTransactionallySafeRWLock& LockObject;
+	FRWScopeLockType LockType;
+};
diff --git a/Engine/Source/Runtime/CoreUObject/Private/UObject/PropertyTypeName.cpp b/Engine/Source/Runtime/CoreUObject/Private/UObject/PropertyTypeName.cpp
--- a/Engine/Source/Runtime/CoreUObject/Private/UObject/PropertyTypeName.cpp
+++ b/Engine/Source/Runtime/CoreUObject/Private/UObject/PropertyTypeName.cpp
@@ -11,6 +11,7 @@
 #include "Misc/StringBuilder.h"
 #include "String/Find.h"
 #include "UObject/Class.h"
+#include "AutoRTFM/AutoRTFM.h"
 
 namespace UE
 {
@@ -161,6 +162,8 @@ FPropertyTypeNameTable::FPropertyTypeNameTable()
 	AllocateBlock(0);
 }
 
+// Open because property type name table is a global cached object.
+UE_AUTORTFM_ALWAYS_OPEN
 int32 FPropertyTypeNameTable::FindOrAddByName(const FPropertyTypeNameNode* First)
 {
 	const FPropertyTypeNameNodeProxy Proxy{First};
@@ -181,6 +184,8 @@ int32 FPropertyTypeNameTable::FindOrAddByName(const FPropertyTypeNameNode* First
 	return Index;
 }
 
+// Open because property type name table is a global cached object.
+UE_AUTORTFM_ALWAYS_OPEN
 const FPropertyTypeNameNode* FPropertyTypeNameTable::ResolveByIndex(int32 Index) const
 {
 	const int32 BlockIndex = Index >> GPropertyTypeNameBlockOffsetBits;
@@ -188,6 +193,8 @@ const FPropertyTypeNameNode* FPropertyTypeNameTable::ResolveByIndex(int32 Index)
 	return Blocks[BlockIndex].load(std::memory_order_relaxed) + BlockOffset;
 }
 
+// Open because property type name table is a global cached object.
+UE_AUTORTFM_ALWAYS_OPEN
 int32 FPropertyTypeNameTable::StoreByIndex(const FPropertyTypeNameNode* Nodes, int32 Count)
 {
 	if (Count == 1 && Nodes->Name.IsNone())
@@ -227,6 +234,8 @@ int32 FPropertyTypeNameTable::StoreByIndex(const FPropertyTypeNameNode* Nodes, i
 	return Index;
 }
 
+// Open because property type name table is a global cached object.
+UE_AUTORTFM_ALWAYS_OPEN
 FPropertyTypeNameNode* FPropertyTypeNameTable::AllocateBlock(int32 BlockIndex)
 {
 	TUniqueLock Lock(BlockAllocationMutex);
diff --git a/Engine/Source/Runtime/CoreUObject/Public/UObject/UObjectAnnotation.h b/Engine/Source/Runtime/CoreUObject/Public/UObject/UObjectAnnotation.h
--- a/Engine/Source/Runtime/CoreUObject/Public/UObject/UObjectAnnotation.h
+++ b/Engine/Source/Runtime/CoreUObject/Public/UObject/UObjectAnnotation.h
@@ -7,8 +7,8 @@
 #pragma once
 
 #include "UObject/UObjectArray.h"
-#include "Misc/ScopeLock.h"
-#include "Misc/ScopeRWLock.h"
+#include "Misc/TransactionallySafeScopeLock.h"
+#include "Misc/TransactionallySafeRWScopeLock.h"
 
 /**
 * FUObjectAnnotationSparse is a helper class that is used to store sparse, slow, temporary, editor only, external 
@@ -39,7 +39,7 @@ public:
 #if !(UE_BUILD_SHIPPING || UE_BUILD_TEST)
 		if (!bAutoRemove)
 		{
-			FScopeLock AnnotationMapLock(&AnnotationMapCritical);
+			FTransactionallySafeScopeLock AnnotationMapLock(&AnnotationMapCritical);
 			// in this case we are only verifying that the external assurances of removal are met
 			check(!AnnotationMap.Find(Object));
 		}
@@ -86,35 +86,22 @@ private:
 		}
 		else
 		{
-			bool bWasEmpty = false;
-			{
-				// We hit this from many locations, for now lets open around adding this, and if we happen to be in a transcation
-				// and abort after we have added let us use the UObject to remove it from the annotation map
-				UE_AUTORTFM_OPEN2
-				{
-					FScopeLock AnnotationMapLock(&AnnotationMapCritical);
-					bWasEmpty = (AnnotationMap.Num() == 0);
-
-					// If we are not in a trasncation update our cache values with the correct information
-					// While in a transaction avoid updating these as it pulls from closed code memory and can cause memory stomps if we abort here
-					if (!AutoRTFM::IsTransactional())
-					{
-						AnnotationCacheKey = Object;
-						AnnotationCacheValue = MoveTemp(LocalAnnotation);
-						AnnotationMap.Add(AnnotationCacheKey, AnnotationCacheValue);
-					}
-					else
-					{
-						AnnotationMap.Add(Object, LocalAnnotation);
-						AnnotationCacheKey = nullptr;
-						AnnotationCacheValue = TAnnotation();
-					}
-				};
+			FTransactionallySafeScopeLock AnnotationMapLock(&AnnotationMapCritical);
+			const bool bWasEmpty = (AnnotationMap.Num() == 0);
 
-				UE_AUTORTFM_ONABORT2(this, Object)
-				{
-					RemoveAnnotation(Object);
-				};
+			// If we are not in a trasncation update our cache values with the correct information
+			// While in a transaction avoid updating these as it pulls from closed code memory and can cause memory stomps if we abort here
+			if (!AutoRTFM::IsTransactional())
+			{
+				AnnotationCacheKey = Object;
+				AnnotationCacheValue = MoveTemp(LocalAnnotation);
+				AnnotationMap.Add(AnnotationCacheKey, AnnotationCacheValue);
+			}
+			else
+			{
+				AnnotationMap.Add(Object, LocalAnnotation);
+				AnnotationCacheKey = nullptr;
+				AnnotationCacheValue = TAnnotation();
 			}
 
 			if (bWasEmpty)
@@ -169,13 +156,14 @@ public:
 		// Avoid holding the lock while we call GUObjectArray.RemoveUObjectDeleteListener as it could deadlock
 		TAnnotation Result;
 		{
-			FScopeLock AnnotationMapLock(&AnnotationMapCritical);
+			FTransactionallySafeScopeLock AnnotationMapLock(&AnnotationMapCritical);
 			AnnotationCacheKey = Object;
 			AnnotationCacheValue = TAnnotation();
 			bHadElements = (AnnotationMap.Num() > 0);
 			AnnotationMap.RemoveAndCopyValue(AnnotationCacheKey, Result);
 			bIsNowEmpty = (AnnotationMap.Num() == 0);
 		}
+
 		if (bHadElements && bIsNowEmpty)
 		{
 			// we are removing the last one, so if we are auto removing or verifying removal, unregister now
@@ -202,7 +190,7 @@ public:
 			bool bIsNowEmpty = false;
 			// Avoid holding the lock while we call GUObjectArray.RemoveUObjectDeleteListener as it could deadlock
 			{
-				FScopeLock AnnotationMapLock(&AnnotationMapCritical);
+				FTransactionallySafeScopeLock AnnotationMapLock(&AnnotationMapCritical);
 				AnnotationCacheKey = Object;
 				AnnotationCacheValue = TAnnotation();
 				bHadElements = (AnnotationMap.Num() > 0);
@@ -231,7 +219,7 @@ public:
 
 		// Avoid holding the lock while we call GUObjectArray.RemoveUObjectDeleteListener as it could deadlock
 		{
-			FScopeLock AnnotationMapLock(&AnnotationMapCritical);
+			FTransactionallySafeScopeLock AnnotationMapLock(&AnnotationMapCritical);
 			AnnotationCacheKey = nullptr;
 			AnnotationCacheValue = TAnnotation();
 			bHadElements = (AnnotationMap.Num() > 0);
@@ -257,26 +245,23 @@ public:
 	FORCEINLINE TAnnotation GetAnnotation(const UObjectBase *Object)
 	{
 		check(Object);
-		TAnnotation Result;
-		UE_AUTORTFM_OPEN2
-		{
-			FScopeLock AnnotationMapLock(&AnnotationMapCritical);
-			if (Object != AnnotationCacheKey)
-			{			
-				AnnotationCacheKey = Object;
-				TAnnotation* Entry = AnnotationMap.Find(AnnotationCacheKey);
-				if (Entry)
-				{
-					AnnotationCacheValue = *Entry;
-				}
-				else
-				{
-					AnnotationCacheValue = TAnnotation();
-				}
+
+		FTransactionallySafeScopeLock AnnotationMapLock(&AnnotationMapCritical);
+		if (Object != AnnotationCacheKey)
+		{			
+			AnnotationCacheKey = Object;
+			TAnnotation* Entry = AnnotationMap.Find(AnnotationCacheKey);
+			if (Entry)
+			{
+				AnnotationCacheValue = *Entry;
 			}
-			Result = AnnotationCacheValue;
-		};
-		return Result;
+			else
+			{
+				AnnotationCacheValue = TAnnotation();
+			}
+		}
+
+		return AnnotationCacheValue;
 	}
 
 	/**
@@ -293,7 +278,7 @@ public:
 	 */
 	void Reserve(int32 ExpectedNumElements)
 	{
-		FScopeLock AnnotationMapLock(&AnnotationMapCritical);
+		FTransactionallySafeScopeLock AnnotationMapLock(&AnnotationMapCritical);
 		AnnotationMap.Empty(ExpectedNumElements);
 	}
 
@@ -308,7 +293,7 @@ private:
 	 * Map from live objects to an annotation
 	 */
 	TMap<const UObjectBase *,TAnnotation> AnnotationMap;
-	FCriticalSection AnnotationMapCritical;
+	FTransactionallySafeCriticalSection AnnotationMapCritical;
 
 	/**
 	 * Key for a one-item cache of the last lookup into AnnotationMap.
@@ -371,7 +356,7 @@ public:
 	 */
 	UObject* Find(const TAnnotation& Annotation)
 	{
-		FScopeLock InverseAnnotationMapLock(&InverseAnnotationMapCritical);
+		FTransactionallySafeScopeLock InverseAnnotationMapLock(&InverseAnnotationMapCritical);
 		checkSlow(!Annotation.IsDefault()); // it is not legal to search for the default annotation
 		return (UObject*)InverseAnnotationMap.FindRef(Annotation);
 	}
@@ -380,7 +365,7 @@ private:
 	template<typename T> 
 	void AddAnnotationInternal(const UObjectBase* Object, T&& Annotation)
 	{
-		FScopeLock InverseAnnotationMapLock(&InverseAnnotationMapCritical);
+		FTransactionallySafeScopeLock InverseAnnotationMapLock(&InverseAnnotationMapCritical);
 		if (Annotation.IsDefault())
 		{
 			RemoveAnnotation(Object); // adding the default annotation is the same as removing an annotation
@@ -425,7 +410,7 @@ public:
 	 */
 	void RemoveAnnotation(const UObjectBase *Object)
 	{
-		FScopeLock InverseAnnotationMapLock(&InverseAnnotationMapCritical);
+		FTransactionallySafeScopeLock InverseAnnotationMapLock(&InverseAnnotationMapCritical);
 		TAnnotation Annotation = this->GetAndRemoveAnnotation(Object);
 		if (Annotation.IsDefault())
 		{
@@ -444,7 +429,7 @@ public:
 	 */
 	void RemoveAllAnnotations()
 	{
-		FScopeLock InverseAnnotationMapLock(&InverseAnnotationMapCritical);
+		FTransactionallySafeScopeLock InverseAnnotationMapLock(&InverseAnnotationMapCritical);
 		Super::RemoveAllAnnotations();
 		InverseAnnotationMap.Empty();
 	}
@@ -459,7 +444,7 @@ private:
 	 * Inverse Map annotation to live object
 	 */
 	TMap<TAnnotation, const UObjectBase *> InverseAnnotationMap;
-	FCriticalSection InverseAnnotationMapCritical;
+	FTransactionallySafeCriticalSection InverseAnnotationMapCritical;
 };
 
 
diff --git a/Engine/Source/Runtime/Engine/Classes/Components/ActorComponent.h b/Engine/Source/Runtime/Engine/Classes/Components/ActorComponent.h
--- a/Engine/Source/Runtime/Engine/Classes/Components/ActorComponent.h
+++ b/Engine/Source/Runtime/Engine/Classes/Components/ActorComponent.h
@@ -15,6 +15,7 @@
 #include "PropertyPairsMap.h"
 #include "ComponentInstanceDataCache.h"
 #include "Experimental/ConcurrentLinearAllocator.h"
+#include "Misc/TransactionallySafeRWScopeLock.h"
 #include "ActorComponent.generated.h"
 
 struct FTypedElementHandle;
@@ -444,7 +445,7 @@ private:
 	TArray<FSimpleMemberReference> UCSModifiedProperties_DEPRECATED;
 #endif
 
-	static ENGINE_API FRWLock AllUCSModifiedPropertiesLock;
+	static ENGINE_API FTransactionallySafeRWLock AllUCSModifiedPropertiesLock;
 	static ENGINE_API TMap<UActorComponent*, TArray<FSimpleMemberReference>> AllUCSModifiedProperties;
 
 public:
diff --git a/Engine/Source/Runtime/Engine/Private/Components/ActorComponent.cpp b/Engine/Source/Runtime/Engine/Private/Components/ActorComponent.cpp
--- a/Engine/Source/Runtime/Engine/Private/Components/ActorComponent.cpp
+++ b/Engine/Source/Runtime/Engine/Private/Components/ActorComponent.cpp
@@ -26,7 +26,6 @@
 #include "HAL/LowLevelMemStats.h"
 #include "Logging/MessageLog.h"
 #include "Misc/MapErrors.h"
-#include "Misc/ScopeRWLock.h"
 #include "Misc/UObjectToken.h"
 #include "Net/Core/PushModel/PushModel.h"
 #include "Net/UnrealNetwork.h"
@@ -431,7 +430,7 @@ UActorComponent::FOnMarkRenderStateDirty UActorComponent::MarkRenderStateDirtyEv
 
 const FString UActorComponent::ComponentTemplateNameSuffix(TEXT("_GEN_VARIABLE"));
 TMap<UActorComponent*, TArray<FSimpleMemberReference>> UActorComponent::AllUCSModifiedProperties;
-FRWLock UActorComponent::AllUCSModifiedPropertiesLock;
+FTransactionallySafeRWLock UActorComponent::AllUCSModifiedPropertiesLock;
 
 UActorComponent::UActorComponent(const FObjectInitializer& ObjectInitializer /*= FObjectInitializer::Get()*/)
 	: Super(ObjectInitializer)
@@ -534,7 +533,7 @@ void UActorComponent::PostLoad()
 		{
 			if (UCSModifiedProperties_DEPRECATED.Num())
 			{
-				FRWScopeLock Lock(AllUCSModifiedPropertiesLock, SLT_Write);
+				FTransactionallySafeRWScopeLock Lock(AllUCSModifiedPropertiesLock, SLT_Write);
 				AllUCSModifiedProperties.Add(this, MoveTemp(UCSModifiedProperties_DEPRECATED));
 			}
 		}
@@ -2483,7 +2482,7 @@ void UActorComponent::DetermineUCSModifiedProperties()
 			}
 		}
 
-		FRWScopeLock Lock(AllUCSModifiedPropertiesLock, SLT_Write);
+		FTransactionallySafeRWScopeLock Lock(AllUCSModifiedPropertiesLock, SLT_Write);
 		if (UCSModifiedProperties.Num() > 0)
 		{
 			AllUCSModifiedProperties.Add(this, MoveTemp(UCSModifiedProperties));
@@ -2497,7 +2496,7 @@ void UActorComponent::DetermineUCSModifiedProperties()
 
 void UActorComponent::GetUCSModifiedProperties(TSet<const FProperty*>& ModifiedProperties) const
 {
-	FRWScopeLock Lock(AllUCSModifiedPropertiesLock, SLT_ReadOnly);
+	FTransactionallySafeRWScopeLock Lock(AllUCSModifiedPropertiesLock, SLT_ReadOnly);
 	if (TArray<FSimpleMemberReference>* UCSModifiedProperties = AllUCSModifiedProperties.Find(this))
 	{
 		for (const FSimpleMemberReference& MemberReference : *UCSModifiedProperties)
@@ -2509,7 +2508,7 @@ void UActorComponent::GetUCSModifiedProperties(TSet<const FProperty*>& ModifiedP
 
 void UActorComponent::RemoveUCSModifiedProperties(const TArray<FProperty*>& Properties)
 {
-	FRWScopeLock Lock(AllUCSModifiedPropertiesLock, SLT_Write);
+	FTransactionallySafeRWScopeLock Lock(AllUCSModifiedPropertiesLock, SLT_Write);
 	if (TArray<FSimpleMemberReference>* UCSModifiedProperties = AllUCSModifiedProperties.Find(this))
 	{
 		for (FProperty* Property : Properties)
@@ -2523,13 +2522,13 @@ void UActorComponent::RemoveUCSModifiedProperties(const TArray<FProperty*>& Prop
 
 void UActorComponent::ClearUCSModifiedProperties()
 {
-	FRWScopeLock Lock(AllUCSModifiedPropertiesLock, SLT_Write);
+	FTransactionallySafeRWScopeLock Lock(AllUCSModifiedPropertiesLock, SLT_Write);
 	AllUCSModifiedProperties.Remove(this);
 }
 
 void UActorComponent::AddReferencedObjects(UObject* InThis, FReferenceCollector& Collector)
 {
-	FRWScopeLock Lock(AllUCSModifiedPropertiesLock, SLT_ReadOnly);
+	FTransactionallySafeRWScopeLock Lock(AllUCSModifiedPropertiesLock, SLT_ReadOnly);
 	if (TArray<FSimpleMemberReference>* UCSModifiedProperties = AllUCSModifiedProperties.Find(CastChecked<UActorComponent>(InThis)))
 	{
 		for (FSimpleMemberReference& MemberReference : *UCSModifiedProperties)
@@ -2585,7 +2584,7 @@ void UActorComponent::Serialize(FArchive& Ar)
 			TArray<FSimpleMemberReference> UCSModifiedProperties;
 			Ar << UCSModifiedProperties;
 
-			FRWScopeLock Lock(AllUCSModifiedPropertiesLock, SLT_Write);
+			FTransactionallySafeRWScopeLock Lock(AllUCSModifiedPropertiesLock, SLT_Write);
 			if (UCSModifiedProperties.Num() > 0)
 			{
 				AllUCSModifiedProperties.Add(this, MoveTemp(UCSModifiedProperties));
@@ -2597,7 +2596,7 @@ void UActorComponent::Serialize(FArchive& Ar)
 		}
 		else
 		{
-			FRWScopeLock Lock(AllUCSModifiedPropertiesLock, SLT_ReadOnly);
+			FTransactionallySafeRWScopeLock Lock(AllUCSModifiedPropertiesLock, SLT_ReadOnly);
 			if (TArray<FSimpleMemberReference>* UCSModifiedProperties = AllUCSModifiedProperties.Find(this))
 			{
 				Ar << *UCSModifiedProperties;
diff --git a/Engine/Source/Runtime/Engine/Private/ContentStreaming.cpp b/Engine/Source/Runtime/Engine/Private/ContentStreaming.cpp
--- a/Engine/Source/Runtime/Engine/Private/ContentStreaming.cpp
+++ b/Engine/Source/Runtime/Engine/Private/ContentStreaming.cpp
@@ -19,6 +19,7 @@
 #include "AudioCompressionSettingsUtils.h"
 #include "VT/VirtualTextureChunkManager.h"
 #include "Rendering/NaniteCoarseMeshStreamingManager.h"
+#include "AutoRTFM/AutoRTFM.h"
 
 #if WITH_EDITOR
 #include "AudioDevice.h"
@@ -370,7 +371,12 @@ FStreamingManagerCollection& IStreamingManager::Get()
 {
 	if (StreamingManagerCollection == nullptr)
 	{
-		StreamingManagerCollection = new FStreamingManagerCollection();
+		// Since this is a lazily created static global variable we create it
+		// in the open.
+		UE_AUTORTFM_OPEN(
+			{
+				StreamingManagerCollection = new FStreamingManagerCollection();
+			});
 	}
 	return *StreamingManagerCollection;
 }
diff --git a/Engine/Source/Runtime/Engine/Private/Tests/AutoRTFM/AutoRTFMActorComponentTests.cpp b/Engine/Source/Runtime/Engine/Private/Tests/AutoRTFM/AutoRTFMActorComponentTests.cpp
--- a/Engine/Source/Runtime/Engine/Private/Tests/AutoRTFM/AutoRTFMActorComponentTests.cpp
+++ b/Engine/Source/Runtime/Engine/Private/Tests/AutoRTFM/AutoRTFMActorComponentTests.cpp
@@ -4,9 +4,11 @@
 #include "Misc/App.h"
 #include "AutoRTFM/AutoRTFM.h"
 #include "AutoRTFMTestActor.h"
+#include "AutoRTFMTestAnotherActor.h"
 #include "AutoRTFMTestBodySetup.h"
 #include "AutoRTFMTestLevel.h"
 #include "AutoRTFMTestObject.h"
+#include "AutoRTFMTestChildActorComponent.h"
 #include "AutoRTFMTestPrimitiveComponent.h"
 #include "PhysicsProxy/SingleParticlePhysicsProxy.h"
 #include "Physics/Experimental/PhysScene_Chaos.h"
@@ -329,6 +331,71 @@ AUTORTFM_ACTOR_COMPONENT_TEST(FSparseDelegate)
 	TEST_CHECK_TRUE(Object->bHitOnComponentPhysicsStateChanged);
 }
 
+AUTORTFM_ACTOR_COMPONENT_TEST(ChildActor)
+{
+	UAutoRTFMTestChildActorComponent* const ChildActorComponent = NewObject<UAutoRTFMTestChildActorComponent>(Actor);
+
+	AAutoRTFMTestAnotherActor* const AnotherActor = NewObject<AAutoRTFMTestAnotherActor>();
+
+	ChildActorComponent->RegisterComponentWithWorld(World);
+
+	ChildActorComponent->ForceActorClass(AnotherActor->GetClass());
+
+	if (nullptr != ChildActorComponent->GetChildActor())
+	{
+		AutoRTFM::ETransactionResult Result = AutoRTFM::Transact([&]
+			{
+				ChildActorComponent->DestroyChildActor();
+				AutoRTFM::AbortTransaction();
+			});
+
+		TEST_CHECK_TRUE(AutoRTFM::ETransactionResult::AbortedByRequest == Result);
+		TEST_CHECK_TRUE(nullptr != ChildActorComponent->GetChildActor());
+
+		Result = AutoRTFM::Transact([&]
+			{
+				ChildActorComponent->DestroyChildActor();
+			});
+
+		TEST_CHECK_TRUE(AutoRTFM::ETransactionResult::Committed == Result);
+		TEST_CHECK_TRUE(nullptr == ChildActorComponent->GetChildActor());
+	}
+
+	AutoRTFM::ETransactionResult Result = AutoRTFM::Transact([&]
+		{
+			ChildActorComponent->CreateChildActor();
+			AutoRTFM::AbortTransaction();
+		});
+
+	TEST_CHECK_TRUE(AutoRTFM::ETransactionResult::AbortedByRequest == Result);
+	TEST_CHECK_TRUE(nullptr == ChildActorComponent->GetChildActor());
+
+	Result = AutoRTFM::Transact([&]
+		{
+			ChildActorComponent->CreateChildActor();
+		});
+
+	TEST_CHECK_TRUE(AutoRTFM::ETransactionResult::Committed == Result);
+	TEST_CHECK_TRUE(nullptr != ChildActorComponent->GetChildActor());
+
+	Result = AutoRTFM::Transact([&]
+		{
+			ChildActorComponent->DestroyChildActor();
+			AutoRTFM::AbortTransaction();
+		});
+
+	TEST_CHECK_TRUE(AutoRTFM::ETransactionResult::AbortedByRequest == Result);
+	TEST_CHECK_TRUE(nullptr != ChildActorComponent->GetChildActor());
+
+	Result = AutoRTFM::Transact([&]
+		{
+			ChildActorComponent->DestroyChildActor();
+		});
+
+	TEST_CHECK_TRUE(AutoRTFM::ETransactionResult::Committed == Result);
+	TEST_CHECK_TRUE(nullptr == ChildActorComponent->GetChildActor());
+}
+
 }  // anonymous namespace
 
 #undef TEST_CHECK_FALSE
diff --git a/Engine/Source/Runtime/Engine/Private/Tests/AutoRTFM/AutoRTFMTestAnotherActor.h b/Engine/Source/Runtime/Engine/Private/Tests/AutoRTFM/AutoRTFMTestAnotherActor.h
new file mode 100644
--- /dev/null
+++ b/Engine/Source/Runtime/Engine/Private/Tests/AutoRTFM/AutoRTFMTestAnotherActor.h
@@ -0,0 +1,15 @@
+// Copyright Epic Games, Inc. All Rights Reserved.
+
+#pragma once
+
+#include "GameFramework/Actor.h"
+#include "AutoRTFMTestAnotherActor.generated.h"
+
+UCLASS()
+class AAutoRTFMTestAnotherActor : public AActor
+{
+    GENERATED_BODY()
+
+public:
+    int Value = 42;
+};
diff --git a/Engine/Source/Runtime/Engine/Private/Tests/AutoRTFM/AutoRTFMTestChildActorComponent.h b/Engine/Source/Runtime/Engine/Private/Tests/AutoRTFM/AutoRTFMTestChildActorComponent.h
new file mode 100644
--- /dev/null
+++ b/Engine/Source/Runtime/Engine/Private/Tests/AutoRTFM/AutoRTFMTestChildActorComponent.h
@@ -0,0 +1,18 @@
+// Copyright Epic Games, Inc. All Rights Reserved.
+
+#pragma once
+
+#include "Components/ChildActorComponent.h"
+#include "AutoRTFMTestChildActorComponent.generated.h"
+
+UCLASS()
+class UAutoRTFMTestChildActorComponent : public UChildActorComponent
+{
+    GENERATED_BODY()
+
+public:
+	void ForceActorClass(UClass* const Class)
+	{
+		SetChildActorClass(Class);
+	}
+};
