Remove the now redundant open around the text formatter static local var.

diff --git a/Engine/Source/Programs/AutoRTFMTests/Private/UECoreTests.cpp b/Engine/Source/Programs/AutoRTFMTests/Private/UECoreTests.cpp
--- a/Engine/Source/Programs/AutoRTFMTests/Private/UECoreTests.cpp
+++ b/Engine/Source/Programs/AutoRTFMTests/Private/UECoreTests.cpp
@@ -4,8 +4,9 @@
 #include "AutoRTFM/AutoRTFM.h"
 #include "Delegates/IDelegateInstance.h"
 #include "HAL/ThreadSingleton.h"
-#include "Internationalization/TextHistory.h"
 #include "Internationalization/TextCache.h"
+#include "Internationalization/TextFormatter.h"
+#include "Internationalization/TextHistory.h"
 #include "Serialization/CustomVersion.h"
 #include "UObject/NameTypes.h"
 #include "UObject/UObjectArray.h"
@@ -1471,3 +1472,27 @@ TEST_CASE("UECore.TransactionallySafeWriteScopeLock")
 		REQUIRE(AutoRTFM::ETransactionResult::Committed == Result);
 	}
 }
+
+TEST_CASE("UECore.FTextFormatPatternDefinition")
+{
+	FTextFormatPatternDefinitionConstPtr Ptr;
+
+	REQUIRE(!Ptr.IsValid());
+
+	AutoRTFM::ETransactionResult Result = AutoRTFM::Transact([&]
+		{
+			Ptr = FTextFormatPatternDefinition::GetDefault().ToSharedPtr();
+			AutoRTFM::AbortTransaction();
+		});
+
+	REQUIRE(AutoRTFM::ETransactionResult::AbortedByRequest == Result);
+	REQUIRE(!Ptr.IsValid());
+
+	Result = AutoRTFM::Transact([&]
+		{
+			Ptr = FTextFormatPatternDefinition::GetDefault().ToSharedPtr();
+		});
+
+	REQUIRE(AutoRTFM::ETransactionResult::Committed == Result);
+	REQUIRE(Ptr.IsValid());
+}
diff --git a/Engine/Source/Runtime/Core/Private/Internationalization/TextFormatter.cpp b/Engine/Source/Runtime/Core/Private/Internationalization/TextFormatter.cpp
--- a/Engine/Source/Runtime/Core/Private/Internationalization/TextFormatter.cpp
+++ b/Engine/Source/Runtime/Core/Private/Internationalization/TextFormatter.cpp
@@ -1008,13 +1008,8 @@ FTextFormatPatternDefinition::FTextFormatPatternDefinition()
 
 FTextFormatPatternDefinitionConstRef FTextFormatPatternDefinition::GetDefault()
 {
-	FTextFormatPatternDefinitionConstRef* Result;
-	AutoRTFM::Open([&Result]()
-	{
-		static FTextFormatPatternDefinitionConstRef DefaultFormatPatternDefinition = MakeShared<FTextFormatPatternDefinition, ESPMode::ThreadSafe>();
-		Result = &DefaultFormatPatternDefinition;
-	});
-	return *Result;
+	static FTextFormatPatternDefinitionConstRef DefaultFormatPatternDefinition = MakeShared<FTextFormatPatternDefinition, ESPMode::ThreadSafe>();
+	return DefaultFormatPatternDefinition;
 }
 
 const FTokenDefinitions& FTextFormatPatternDefinition::GetTextFormatDefinitions() const
