Make setting and clearing UObject flags transactionally safe.

diff --git a/Engine/Source/Programs/AutoRTFMTests/Private/UObjectTests.cpp b/Engine/Source/Programs/AutoRTFMTests/Private/UObjectTests.cpp
--- a/Engine/Source/Programs/AutoRTFMTests/Private/UObjectTests.cpp
+++ b/Engine/Source/Programs/AutoRTFMTests/Private/UObjectTests.cpp
@@ -236,3 +236,67 @@ TEST_CASE("UObject.AnnotationMap")
 
 	REQUIRE(!AnnotationMap.GetAnnotation(Key).IsDefault());
 }
+
+TEST_CASE("UObject.AtomicallySetFlags")
+{
+	UObject* const Object = NewObject<UMyAutoRTFMTestObject>();
+
+	constexpr EObjectFlags OldFlags = EObjectFlags::RF_Public | EObjectFlags::RF_Transient;
+	constexpr EObjectFlags FlagsToAdd = EObjectFlags::RF_Transient | EObjectFlags::RF_AllocatedInSharedPage;
+
+	// We need to ensure we cover the case where we are adding a flag that is already there
+	// and thus cannot just wipe that out if we abort!
+	Object->AtomicallyClearFlags(FlagsToAdd);
+	Object->AtomicallySetFlags(OldFlags);
+
+	REQUIRE(Object->HasAllFlags(OldFlags) & !Object->HasAllFlags(FlagsToAdd));
+
+	AutoRTFM::ETransactionResult Result = AutoRTFM::Transact([&]
+		{
+			Object->AtomicallySetFlags(FlagsToAdd);
+			AutoRTFM::AbortTransaction();
+		});
+
+	REQUIRE(AutoRTFM::ETransactionResult::AbortedByRequest == Result);
+	REQUIRE(Object->HasAllFlags(OldFlags) & !Object->HasAllFlags(FlagsToAdd));
+
+	Result = AutoRTFM::Transact([&]
+		{
+			Object->AtomicallySetFlags(FlagsToAdd);
+		});
+
+	REQUIRE(AutoRTFM::ETransactionResult::Committed == Result);
+	REQUIRE(Object->HasAllFlags(OldFlags) & Object->HasAllFlags(FlagsToAdd));
+}
+
+TEST_CASE("UObject.AtomicallyClearFlags")
+{
+	UObject* const Object = NewObject<UMyAutoRTFMTestObject>();
+
+	constexpr EObjectFlags OldFlags = EObjectFlags::RF_Public | EObjectFlags::RF_Transient;
+	constexpr EObjectFlags FlagsToClear = EObjectFlags::RF_Transient | EObjectFlags::RF_AllocatedInSharedPage;
+
+	// We need to ensure we cover the case where we are adding a flag that is already there
+	// and thus cannot just wipe that out if we abort!
+	Object->AtomicallyClearFlags(FlagsToClear);
+	Object->AtomicallySetFlags(OldFlags);
+
+	REQUIRE(Object->HasAllFlags(OldFlags) & !Object->HasAllFlags(FlagsToClear));
+
+	AutoRTFM::ETransactionResult Result = AutoRTFM::Transact([&]
+		{
+			Object->AtomicallyClearFlags(FlagsToClear);
+			AutoRTFM::AbortTransaction();
+		});
+
+	REQUIRE(AutoRTFM::ETransactionResult::AbortedByRequest == Result);
+	REQUIRE(Object->HasAllFlags(OldFlags) & !Object->HasAllFlags(FlagsToClear));
+
+	Result = AutoRTFM::Transact([&]
+		{
+			Object->AtomicallyClearFlags(FlagsToClear);
+		});
+
+	REQUIRE(AutoRTFM::ETransactionResult::Committed == Result);
+	REQUIRE(Object->HasAnyFlags(OldFlags) & !Object->HasAllFlags(FlagsToClear));
+}
diff --git a/Engine/Source/Runtime/CoreUObject/Public/UObject/UObjectBase.h b/Engine/Source/Runtime/CoreUObject/Public/UObject/UObjectBase.h
--- a/Engine/Source/Runtime/CoreUObject/Public/UObject/UObjectBase.h
+++ b/Engine/Source/Runtime/CoreUObject/Public/UObject/UObjectBase.h
@@ -238,7 +238,23 @@ public:
 			return;
 		}
 
-		FPlatformAtomics::InterlockedOr((int32*)&ObjectFlags, FlagsToAdd);
+		UE_AUTORTFM_OPEN(
+			{
+				FPlatformAtomics::InterlockedOr((int32*)&ObjectFlags, FlagsToAdd);
+			});
+
+		// If we abort we undo setting the flags we just set.
+		AutoRTFM::OnAbort([this, OldFlags, FlagsToAdd]
+			{
+				int32 MaskFlags = OldFlags;
+
+				// Now just extract out the old flags that mattered (the ones we were setting).
+				MaskFlags &= FlagsToAdd;
+				// And unmask the flags we didn't mention.
+				MaskFlags |= ~FlagsToAdd;
+
+				FPlatformAtomics::InterlockedAnd((int32*)&ObjectFlags, MaskFlags);
+			});
 	}
 
 	/**
@@ -255,7 +271,21 @@ public:
 			return;
 		}
 
-		FPlatformAtomics::InterlockedAnd((int32*)&ObjectFlags, ~FlagsToClear);
+		UE_AUTORTFM_OPEN(
+			{
+				FPlatformAtomics::InterlockedAnd((int32*)&ObjectFlags, ~FlagsToClear);
+			});
+
+		// If we abort we undo clearing the flags we just unset.
+		AutoRTFM::OnAbort([this, OldFlags, FlagsToClear]
+			{
+				int32 MaskFlags = OldFlags;
+
+				// Now just extract out the old flags that mattered (the ones we were setting).
+				MaskFlags &= FlagsToClear;
+
+				FPlatformAtomics::InterlockedOr((int32*)&ObjectFlags, MaskFlags);
+			});
 	}
 
 	static void PrefetchClass(UObject* Object) { FPlatformMisc::Prefetch(Object, offsetof(UObjectBase, ClassPrivate)); }
