Fix the FString::Printf problem which came down to a fun ODR violation in the MSVC CRT that AutoRTFM exposed.

diff --git a/Engine/Source/Programs/AutoRTFMTests/Private/Tests.cpp b/Engine/Source/Programs/AutoRTFMTests/Private/Tests.cpp
--- a/Engine/Source/Programs/AutoRTFMTests/Private/Tests.cpp
+++ b/Engine/Source/Programs/AutoRTFMTests/Private/Tests.cpp
@@ -447,3 +447,29 @@ TEST_CASE("Tests.AVXMaskedStore")
 	}
 }
 #endif // PLATFORM_CPU_X86_FAMILY
+
+TEST_CASE("Tests.__local_stdio_printf_options")
+{
+	unsigned __int64* NonTransactional = __local_stdio_printf_options();
+	unsigned __int64* Transactional = nullptr;
+
+	AutoRTFM::Commit([&]
+		{
+			Transactional = __local_stdio_printf_options();
+		});
+
+	REQUIRE(NonTransactional == Transactional);
+}
+
+TEST_CASE("Tests.__local_stdio_scanf_options")
+{
+	unsigned __int64* NonTransactional = __local_stdio_scanf_options();
+	unsigned __int64* Transactional = nullptr;
+
+	AutoRTFM::Commit([&]
+		{
+			Transactional = __local_stdio_scanf_options();
+		});
+
+	REQUIRE(NonTransactional == Transactional);
+}
diff --git a/Engine/Source/Programs/AutoRTFMTests/Private/UECoreTests.cpp b/Engine/Source/Programs/AutoRTFMTests/Private/UECoreTests.cpp
--- a/Engine/Source/Programs/AutoRTFMTests/Private/UECoreTests.cpp
+++ b/Engine/Source/Programs/AutoRTFMTests/Private/UECoreTests.cpp
@@ -1496,3 +1496,18 @@ TEST_CASE("UECore.FTextFormatPatternDefinition")
 	REQUIRE(AutoRTFM::ETransactionResult::Committed == Result);
 	REQUIRE(Ptr.IsValid());
 }
+
+TEST_CASE("UECore.FString")
+{
+	SECTION("Printf")
+	{
+		FString String;
+
+		AutoRTFM::Commit([&]
+			{
+				String = FString::Printf(TEXT("Foo '%s' Bar"), TEXT("Stuff"));
+			});
+
+		REQUIRE(String == "Foo 'Stuff' BAR");
+	}
+}
