[AutoRTFMTests] Make tests pass on linux

* Use fopen() as a no-RTFM function instead of std::atomic_bool, as the latter has an inlined implementation on linux.
* Add missing tests for modf() and modff(). Add the missing modfl() registration which these caught.
* Fix RTFM_vswprintf() count calculations for non-windows. vswprintf() is horribly non-portable, and returns -1 if the buffer is too short.
* Adjust the printf tests to handle other non-portablity issues.
* Wrap wide string literals with std::wstring_view() when in REQUIRES(), as catch2 prints a bunch of '?' on error otherwise.

diff --git a/Engine/Source/Programs/AutoRTFMTests/Private/Abort.cpp b/Engine/Source/Programs/AutoRTFMTests/Private/Abort.cpp
--- a/Engine/Source/Programs/AutoRTFMTests/Private/Abort.cpp
+++ b/Engine/Source/Programs/AutoRTFMTests/Private/Abort.cpp
@@ -525,15 +525,19 @@ TEST_CASE("Abort.OnAbortTiming")
 	REQUIRE(bOnAbortRan == true);
 }
 
+static void FnHasNoClosed()
+{
+	(void)fopen("fopen() is not supported in a closed transaction", "rb");
+}
+
 TEST_CASE("Abort.Language")
 {
 	bool bTouched = false;
-	std::atomic_bool bThingie = false;
 
 	const AutoRTFM::ETransactionResult Result = AutoRTFM::Transact([&]
 		{
 			bTouched = true;
-			bThingie = true;
+			FnHasNoClosed();
 		});
 
 	REQUIRE(AutoRTFM::ETransactionResult::AbortedByLanguage == Result);
@@ -543,7 +547,6 @@ TEST_CASE("Abort.Language")
 TEST_CASE("Abort.LanguageThroughOpen")
 {
 	bool bTouched = false;
-	std::atomic_bool bThingie = false;
 
 	const AutoRTFM::ETransactionResult Result = AutoRTFM::Transact([&]
 		{
@@ -553,7 +556,7 @@ TEST_CASE("Abort.LanguageThroughOpen")
 				{
 					const AutoRTFM::EContextStatus Status = AutoRTFM::Close([&]
 						{
-							bThingie = true;
+							FnHasNoClosed();
 						});
 
 					REQUIRE(AutoRTFM::EContextStatus::AbortedByLanguage == Status);
diff --git a/Engine/Source/Programs/AutoRTFMTests/Private/CString.cpp b/Engine/Source/Programs/AutoRTFMTests/Private/CString.cpp
--- a/Engine/Source/Programs/AutoRTFMTests/Private/CString.cpp
+++ b/Engine/Source/Programs/AutoRTFMTests/Private/CString.cpp
@@ -630,7 +630,7 @@ TEST_CASE("CString.wcscpy")
 		CLANG_END_DISABLE_OPTIMIZATIONS
 
 		REQUIRE(AutoRTFM::ETransactionResult::AbortedByRequest == Result);
-		REQUIRE(L"Doggie says woof____" == std::wstring_view(To));
+		REQUIRE(std::wstring_view(L"Doggie says woof____") == std::wstring_view(To));
 	}
 
 	SECTION("With Commit")
@@ -639,7 +639,7 @@ TEST_CASE("CString.wcscpy")
 		AutoRTFM::Commit([&] { wcscpy(To, From); });
 		CLANG_END_DISABLE_OPTIMIZATIONS
 
-		REQUIRE(L"Kittie says meow" == std::wstring_view(To));
+		REQUIRE(std::wstring_view(L"Kittie says meow") == std::wstring_view(To));
 	}
 
 	MSVC_END_DISABLE_UNSAFE_FN
@@ -663,7 +663,7 @@ TEST_CASE("CString.wcsncpy")
 		CLANG_END_DISABLE_OPTIMIZATIONS
 
 		REQUIRE(AutoRTFM::ETransactionResult::AbortedByRequest == Result);
-		REQUIRE(L"Doggie says woof" == std::wstring_view(To));
+		REQUIRE(std::wstring_view(L"Doggie says woof") == std::wstring_view(To));
 	}
 
 	SECTION("With Commit")
@@ -672,7 +672,7 @@ TEST_CASE("CString.wcsncpy")
 		AutoRTFM::Commit([&] { wcsncpy(To, From, 6); });
 		CLANG_END_DISABLE_OPTIMIZATIONS
 
-		REQUIRE(L"Kittie says woof" == std::wstring_view(To));
+		REQUIRE(std::wstring_view(L"Kittie says woof") == std::wstring_view(To));
 	}
 
 	MSVC_END_DISABLE_UNSAFE_FN
@@ -1165,7 +1165,7 @@ TEST_CASE("CString.swprintf")
 
 			REQUIRE(AutoRTFM::ETransactionResult::AbortedByRequest == Result);
 			REQUIRE(0 == Count);
-			REQUIRE(L"_____________________" == std::wstring_view(Buffer));
+			REQUIRE(std::wstring_view(L"_____________________") == std::wstring_view(Buffer));
 		}
 
 		SECTION("With Commit")
@@ -1180,7 +1180,7 @@ TEST_CASE("CString.swprintf")
 			CLANG_END_DISABLE_OPTIMIZATIONS
 
 			REQUIRE(16 == Count);
-			REQUIRE(L"cat says 'meow'!" == std::wstring_view(Buffer));
+			REQUIRE(std::wstring_view(L"cat says 'meow'!") == std::wstring_view(Buffer));
 		}
 	}
 
@@ -1202,7 +1202,7 @@ TEST_CASE("CString.swprintf")
 
 			REQUIRE(AutoRTFM::ETransactionResult::AbortedByRequest == Result);
 			REQUIRE(0 == Count);
-			REQUIRE(L"_____________________" == std::wstring_view(Buffer));
+			REQUIRE(std::wstring_view(L"_____________________") == std::wstring_view(Buffer));
 		}
 
 		SECTION("With Commit")
@@ -1217,7 +1217,7 @@ TEST_CASE("CString.swprintf")
 			CLANG_END_DISABLE_OPTIMIZATIONS
 
 			REQUIRE(0 > Count);
-			REQUIRE(L"cat say" == std::wstring_view(Buffer));
+			REQUIRE(std::wstring_view(L"cat say") == std::wstring_view(Buffer, 7));
 		}
 
 		CLANG_END_DISABLE_WARN_FORMAT_TRUNCATED
@@ -1252,7 +1252,7 @@ TEST_CASE("CString.swprintf")
 			});
 			CLANG_END_DISABLE_OPTIMIZATIONS
 
-			REQUIRE(16 == Count);
+			REQUIRE(((0 > Count) || (16 == Count)));
 		}
 	}
 
@@ -1271,7 +1271,7 @@ TEST_CASE("CString.swprintf")
 
 		REQUIRE(AutoRTFM::ETransactionResult::AbortedByLanguage == Result);
 		REQUIRE(0 == Count);
-		REQUIRE(L"_____________________" == std::wstring_view(Buffer));
+		REQUIRE(std::wstring_view(L"_____________________") == std::wstring_view(Buffer));
 		REQUIRE_THAT(WarningContext.GetWarnings(),
 			Catch::Matchers::VectorContains(FString(kPercentNWarning)));
 	}
@@ -1292,7 +1292,7 @@ TEST_CASE("CString.swprintf")
 
 			REQUIRE(AutoRTFM::ETransactionResult::AbortedByRequest == Result);
 			REQUIRE(0 == Count);
-			REQUIRE(L"_____________________" == std::wstring_view(Buffer));
+			REQUIRE(std::wstring_view(L"_____________________") == std::wstring_view(Buffer));
 		}
 
 		SECTION("With Commit")
@@ -1307,7 +1307,7 @@ TEST_CASE("CString.swprintf")
 			CLANG_END_DISABLE_OPTIMIZATIONS
 
 			REQUIRE(2 == Count);
-			REQUIRE(L"%n" == std::wstring_view(Buffer));
+			REQUIRE(std::wstring_view(L"%n") == std::wstring_view(Buffer));
 		}
 	}
 }
@@ -1366,6 +1366,7 @@ TEST_CASE("CString.printf")
 
 TEST_CASE("CString.wprintf")
 {
+#if PLATFORM_WINDOWS // wprintf() can error on linux
 	SECTION("With Abort")
 	{
 		int Count = 0;
@@ -1387,14 +1388,16 @@ TEST_CASE("CString.wprintf")
 		int Count = 0;
 
 		CLANG_BEGIN_DISABLE_OPTIMIZATIONS
-		AutoRTFM::Commit([&]
+		AutoRTFM::ETransactionResult Result = AutoRTFM::Transact([&]
 		{
 			Count = wprintf(L"AutoRTFM wprintf test: %%n\n");
 		});
 		CLANG_END_DISABLE_OPTIMIZATIONS
 
+		REQUIRE(AutoRTFM::ETransactionResult::Committed == Result);
 		REQUIRE(26 == Count);
 	}
+#endif // PLATFORM_WINDOWS
 
 	SECTION("PercentN")
 	{
diff --git a/Engine/Source/Programs/AutoRTFMTests/Private/Math.cpp b/Engine/Source/Programs/AutoRTFMTests/Private/Math.cpp
--- a/Engine/Source/Programs/AutoRTFMTests/Private/Math.cpp
+++ b/Engine/Source/Programs/AutoRTFMTests/Private/Math.cpp
@@ -3,6 +3,7 @@
 #include "Catch2Includes.h"
 #include <AutoRTFM/AutoRTFM.h>
 #include <math.h>
+#include <utility>
 
 namespace
 {
@@ -44,221 +45,253 @@ void Check(FUNC&& Function)
 
 }  // anonymous namespace
 
-TEST_CASE("Math.Sqrt")
+TEST_CASE("Math.sqrt")
 {
 	SECTION("float")  { Check([]{ return sqrt(0.42f); }); };
 	SECTION("double") { Check([]{ return sqrt(0.42); }); };
 }
 
-TEST_CASE("Math.Sqrtf")
+TEST_CASE("Math.sqrtf")
 {
 	Check([]{ return sqrtf(0.42f); });
 }
 
-TEST_CASE("Math.Sin")
+TEST_CASE("Math.sin")
 {
 	SECTION("float")  { Check([]{ return sin(0.42f); }); }
 	SECTION("double") { Check([]{ return sin(0.42); }); }
 }
 
-TEST_CASE("Math.Sinf")
+TEST_CASE("Math.sinf")
 {
 	Check([]{ return sinf(0.42f); });
 }
 
-TEST_CASE("Math.Cos")
+TEST_CASE("Math.cos")
 {
 	SECTION("float")  { Check([]{ return cos(0.42f); }); }
 	SECTION("double") { Check([]{ return cos(0.42); }); }
 }
 
-TEST_CASE("Math.Cosf")
+TEST_CASE("Math.cosf")
 {
 	Check([]{ return cosf(0.42f); });
 }
 
-TEST_CASE("Math.Tan")
+TEST_CASE("Math.tan")
 {
 	SECTION("float")  { Check([]{ return tan(0.42f); }); }
 	SECTION("double") { Check([]{ return tan(0.42); }); }
 }
 
-TEST_CASE("Math.Tanf")
+TEST_CASE("Math.tanf")
 {
 	Check([]{ return tanf(0.42f); });
 }
 
-TEST_CASE("Math.Asin")
+TEST_CASE("Math.asin")
 {
 	SECTION("float")  { Check([]{ return asin(0.42f); }); }
 	SECTION("double") { Check([]{ return asin(0.42); }); }
 }
 
-TEST_CASE("Math.Asinf")
+TEST_CASE("Math.asinf")
 {
 	Check([]{ return asinf(0.42f); });
 }
 
-TEST_CASE("Math.Acos")
+TEST_CASE("Math.acos")
 {
 	SECTION("float")  { Check([]{ return acos(0.42f); }); }
 	SECTION("double") { Check([]{ return acos(0.42); }); }
 }
 
-TEST_CASE("Math.Acosf")
+TEST_CASE("Math.acosf")
 {
 	Check([]{ return acosf(0.42f); });
 }
 
-TEST_CASE("Math.Atan")
+TEST_CASE("Math.atan")
 {
 	SECTION("float")  { Check([]{ return atan(0.42f); }); }
 	SECTION("double") { Check([]{ return atan(0.42); }); }
 }
 
-TEST_CASE("Math.Atanf")
+TEST_CASE("Math.atanf")
 {
 	Check([]{ return atanf(0.42f); });
 }
 
-TEST_CASE("Math.Atan2")
+TEST_CASE("Math.atan2")
 {
 	SECTION("float")  { Check([]{ return atan2(0.42f, 0.42f); }); }
 	SECTION("double") { Check([]{ return atan2(0.42, 0.42); }); }
 }
 
-TEST_CASE("Math.Atan2f")
+TEST_CASE("Math.atan2f")
 {
 	Check([]{ return atan2f(0.42f, 0.24f); });
 }
 
-TEST_CASE("Math.Sinh")
+TEST_CASE("Math.sinh")
 {
 	SECTION("float")  { Check([]{ return sinh(0.42f); }); }
 	SECTION("double") { Check([]{ return sinh(0.42); }); }
 }
 
-TEST_CASE("Math.Sinhf")
+TEST_CASE("Math.sinhf")
 {
 	Check([]{ return sinhf(0.42f); });
 }
 
-TEST_CASE("Math.Cosh")
+TEST_CASE("Math.cosh")
 {
 	SECTION("float")  { Check([]{ return cosh(0.42f); }); }
 	SECTION("double") { Check([]{ return cosh(0.42); }); }
 }
 
-TEST_CASE("Math.Coshf")
+TEST_CASE("Math.coshf")
 {
 	Check([]{ return coshf(0.42f); });
 }
 
-TEST_CASE("Math.Tanh")
+TEST_CASE("Math.tanh")
 {
 	SECTION("float")  { Check([]{ return tanh(0.42f); }); }
 	SECTION("double") { Check([]{ return tanh(0.42); }); }
 }
 
-TEST_CASE("Math.Tanhf")
+TEST_CASE("Math.tanhf")
 {
 	Check([]{ return tanhf(0.42f); });
 }
 
-TEST_CASE("Math.Asinh")
+TEST_CASE("Math.asinh")
 {
 	SECTION("float")  { Check([]{ return asinh(0.42f); }); }
 	SECTION("double") { Check([]{ return asinh(0.42); }); }
 }
 
-TEST_CASE("Math.Asinhf")
+TEST_CASE("Math.asinhf")
 {
 	Check([]{ return asinhf(0.42f); });
 }
 
-TEST_CASE("Math.Acosh")
+TEST_CASE("Math.acosh")
 {
 	SECTION("float")  { Check([]{ return acosh(4.2f); }); }
 	SECTION("double") { Check([]{ return acosh(4.2); }); }
 }
 
-TEST_CASE("Math.Acoshf")
+TEST_CASE("Math.acoshf")
 {
 	Check([]{ return acoshf(4.2f); });
 }
 
-TEST_CASE("Math.Atanh")
+TEST_CASE("Math.atanh")
 {
 	SECTION("float")  { Check([]{ return atanh(0.42f); }); }
 	SECTION("double") { Check([]{ return atanh(0.42); }); }
 }
 
-TEST_CASE("Math.Atanhf")
+TEST_CASE("Math.atanhf")
 {
 	Check([]{ return atanhf(0.42f); });
 }
 
-TEST_CASE("Math.Exp")
+TEST_CASE("Math.exp")
 {
 	SECTION("float")  { Check([]{ return exp(0.42f); }); }
 	SECTION("double") { Check([]{ return exp(0.42); }); }
 }
 
-TEST_CASE("Math.Expf")
+TEST_CASE("Math.expf")
 {
 	Check([]{ return expf(0.42f); });
 }
 
-TEST_CASE("Math.Log")
+TEST_CASE("Math.log")
 {
 	SECTION("float")  { Check([]{ return log(0.42f); }); }
 	SECTION("double") { Check([]{ return log(0.42); }); }
 }
 
-TEST_CASE("Math.Pow")
+TEST_CASE("Math.pow")
 {
 	SECTION("float")  { Check([]{ return pow(0.42f, 0.42f); }); }
 	SECTION("double") { Check([]{ return pow(0.42, 0.42); }); }
 }
 
-TEST_CASE("Math.Powf")
+TEST_CASE("Math.powf")
 {
 	Check([]{ return powf(0.42f, 0.24f); });
 }
 
-TEST_CASE("Math.Logf")
+TEST_CASE("Math.logf")
 {
 	Check([]{ return logf(0.42f); });
 }
 
-TEST_CASE("Math.Llrint")
+TEST_CASE("Math.llrint")
 {
 	SECTION("float")  { Check([]{ return llrint(0.42f); }); }
 	SECTION("double") { Check([]{ return llrint(0.42); }); }
 }
 
-TEST_CASE("Math.Llrintf")
+TEST_CASE("Math.llrintf")
 {
 	Check([]{ return llrintf(0.42f); });
 }
 
-TEST_CASE("Math.Fmod")
+TEST_CASE("Math.fmod")
 {
 	SECTION("float")  { Check([]{ return fmod(0.42f, 0.42f); }); }
 	SECTION("double") { Check([]{ return fmod(0.42, 0.42); }); }
 }
 
-TEST_CASE("Math.Fmodf")
+TEST_CASE("Math.fmodf")
 {
 	Check([]{ return fmodf(0.42f, 0.24f); });
 }
 
-TEST_CASE("Math.Fmodl")
+TEST_CASE("Math.fmodl")
 {
 	Check([]{ return fmodl(0.42, 0.24); });
 }
 
-TEST_CASE("Math.Rand")
+TEST_CASE("Math.modf")
+{
+	SECTION("float")
+	{ 
+		Check([]
+		{
+			float A = 0;
+			float B = modf(4.2f, &A);
+			return std::make_pair(A, B); 
+		});
+	}
+	SECTION("double")
+	{
+		Check([]
+		{
+			long double A = 0;
+			long double B = modf(4.2, &A);
+			return std::make_pair(A, B); 
+		});
+	}
+}
+
+TEST_CASE("Math.modff")
+{
+	Check([]
+	{
+		float A = 0;
+		float B = modff(4.2f, &A);
+		return std::make_pair(A, B); 
+	});
+}
+
+TEST_CASE("Math.rand")
 {
 	// it's near impossible to test the value returned by rand(), so just check
 	// that calling it a transaction doesn't explode.
diff --git a/Engine/Source/Programs/AutoRTFMTests/Private/Tests.cpp b/Engine/Source/Programs/AutoRTFMTests/Private/Tests.cpp
--- a/Engine/Source/Programs/AutoRTFMTests/Private/Tests.cpp
+++ b/Engine/Source/Programs/AutoRTFMTests/Private/Tests.cpp
@@ -448,6 +448,8 @@ TEST_CASE("Tests.AVXMaskedStore")
 }
 #endif // PLATFORM_CPU_X86_FAMILY
 
+#if PLATFORM_WINDOWS
+
 TEST_CASE("Tests.__local_stdio_printf_options")
 {
 	unsigned __int64* NonTransactional = __local_stdio_printf_options();
@@ -473,3 +475,5 @@ TEST_CASE("Tests.__local_stdio_scanf_options")
 
 	REQUIRE(NonTransactional == Transactional);
 }
+
+#endif // PLATFORM_WINDOWS
diff --git a/Engine/Source/Runtime/Core/Private/AutoRTFM/StdLibSupport.cpp b/Engine/Source/Runtime/Core/Private/AutoRTFM/StdLibSupport.cpp
--- a/Engine/Source/Runtime/Core/Private/AutoRTFM/StdLibSupport.cpp
+++ b/Engine/Source/Runtime/Core/Private/AutoRTFM/StdLibSupport.cpp
@@ -26,6 +26,23 @@ namespace AutoRTFM
 namespace
 {
 
+// A helper that opens a FILE to "/dev/null" on first call to Get()
+// and automatically closes the file on static destruction.
+class FNullFile
+{
+public:
+    static FILE* Get()
+    {
+        static FNullFile Instance;
+        return Instance.File;
+    }
+
+private:
+    FNullFile() : File(fopen("/dev/null", "wb")) {}
+    ~FNullFile() { fclose(File); }
+    FILE* const File;
+};
+
 void ThrowErrorFormatContainsPercentN()
 {
     UE_LOG(LogAutoRTFM, Warning, TEXT("AutoRTFM does not support format strings containing '%%n'"));
@@ -303,6 +320,7 @@ UE_AUTORTFM_REGISTER_SELF_FUNCTION(fmodf);
 UE_AUTORTFM_REGISTER_SELF_FUNCTION(fmodl);
 UE_AUTORTFM_REGISTER_SELF_FUNCTION(rand);
 UE_AUTORTFM_REGISTER_SELF_FUNCTION(modff);
+UE_AUTORTFM_REGISTER_SELF_FUNCTION(modfl);
 
 // FIXME: Does not currently support %n format specifiers.
 int RTFM_vsnprintf(char* Str, size_t Size, const char* Format, va_list ArgList)
@@ -337,11 +355,18 @@ int RTFM_vswprintf(wchar_t* Str, size_t Size, const wchar_t* Format, va_list Arg
         va_list ArgList2;
         va_copy(ArgList2, ArgList);
 
-        FContext* Context = FContext::Get();
+#if PLATFORM_WINDOWS
         int Count = vswprintf(nullptr, 0, Format, ArgList2);
-        if (Count >= 0)
+#else
+        // vswprintf(nullptr, 0, ...) will return -1.
+        int Count = vfwprintf(FNullFile::Get(), Format, ArgList2);
+#endif
+
+        size_t NumChars = FMath::Min<size_t>(Size, (1 + FMath::Max(Count, 0)));
+        size_t NumBytes = NumChars * sizeof(wchar_t);
+        if (NumBytes >= 0)
         {
-            size_t NumBytes = FMath::Min<size_t>(Size, (1 + Count)) * sizeof(wchar_t);
+            FContext* Context = FContext::Get();
             Context->RecordWrite(Str, NumBytes);
         }
     }
@@ -352,6 +377,21 @@ UE_AUTORTFM_REGISTER_OPEN_FUNCTION_EXPLICIT(
     static_cast<int(*)(wchar_t*, size_t, const wchar_t*, va_list)>(&vswprintf),
     RTFM_vswprintf);
 
+// FIXME: Does not currently support %n format specifiers.
+int RTFM_swprintf(wchar_t* Buffer, size_t BufferCount, wchar_t const* Format, ...)
+{
+    va_list ArgList;
+
+    va_start(ArgList, Format);
+    int Count = RTFM_vswprintf(Buffer, BufferCount, Format, ArgList);
+    va_end(ArgList);
+
+    return Count;
+}
+UE_AUTORTFM_REGISTER_OPEN_FUNCTION_EXPLICIT(
+    static_cast<int(*)(wchar_t*, size_t, wchar_t const*, ...)>(&swprintf),
+    RTFM_swprintf);
+
 // FIXME: Does not currently support %n format specifiers.
 int RTFM_snprintf(char* Str, size_t Size, const char* Format, ...)
 {
