autortfmtests: suppress log pruning by passing -Multiprocess into the engine cmdline, since we run variants of the test suite in parallel using multiple processes and that causes races with pruning old log files

diff --git a/Engine/Source/Programs/AutoRTFMTests/Private/AutoRTFMTests.cpp b/Engine/Source/Programs/AutoRTFMTests/Private/AutoRTFMTests.cpp
--- a/Engine/Source/Programs/AutoRTFMTests/Private/AutoRTFMTests.cpp
+++ b/Engine/Source/Programs/AutoRTFMTests/Private/AutoRTFMTests.cpp
@@ -41,7 +41,8 @@ int main(int ArgC, const char* ArgV[]) {
 		AutoRTFM::ForTheRuntime::SetRetryTransaction(AutoRTFM::ForTheRuntime::EAutoRTFMRetryTransactionState::RetryNonNested);
 	}
 
-	GEngineLoop.PreInit(0, nullptr);
+	const TCHAR* CommandLine = TEXT("-Multiprocess");
+	GEngineLoop.PreInit(CommandLine);
 	FModuleManager::Get().StartProcessingNewlyLoadedObjects();
 
 	// Enable AutoRTFM.
