[AutoRTFM] Carefully destruct OnCommit / OnAbort functions

Abort tasks should be destructed immediately after being called, otherwise their destructors may depend on memory freed by a later task.
Commit tasks should be destructed before the abort handlers, again as they might depend on memory that is freed.

diff --git a/Engine/Source/Programs/AutoRTFMTests/Private/SharedPointerTests.cpp b/Engine/Source/Programs/AutoRTFMTests/Private/SharedPointerTests.cpp
--- a/Engine/Source/Programs/AutoRTFMTests/Private/SharedPointerTests.cpp
+++ b/Engine/Source/Programs/AutoRTFMTests/Private/SharedPointerTests.cpp
@@ -184,3 +184,45 @@ TEST_CASE("SharedPointer.AbortNestedTransactionWithPlacementNewlyAllocated")
 
 	REQUIRE(42 == Result);
 }
+
+TEST_CASE("SharedPointer.OnCommitCapturesSharedPtr")
+{
+	AutoRTFM::Transact([&]
+	{
+		TSharedPtr<int> Shared{new int};
+		AutoRTFM::OnCommit([Shared] {});
+		AutoRTFM::AbortTransaction();
+	});
+}
+
+TEST_CASE("SharedPointer.OnCommitCapturesWeakPtr")
+{
+	AutoRTFM::Transact([&]
+	{
+		TSharedPtr<int> Shared{new int};
+		TWeakPtr<int> Weak{Shared};
+		AutoRTFM::OnCommit([Weak] {});
+		AutoRTFM::AbortTransaction();
+	});
+}
+
+TEST_CASE("SharedPointer.OnAbortCapturesSharedPtr")
+{
+	AutoRTFM::Transact([&]
+	{
+		TSharedPtr<int> Shared{new int};
+		AutoRTFM::OnAbort([Shared] {});
+		AutoRTFM::AbortTransaction();
+	});
+}
+
+TEST_CASE("SharedPointer.OnAbortCapturesWeakPtr")
+{
+	AutoRTFM::Transact([&]
+	{
+		TSharedPtr<int> Shared{new int};
+		TWeakPtr<int> Weak{Shared};
+		AutoRTFM::OnAbort([Weak] {});
+		AutoRTFM::AbortTransaction();
+	});
+}
diff --git a/Engine/Source/Runtime/Core/Private/AutoRTFM/TaskArray.h b/Engine/Source/Runtime/Core/Private/AutoRTFM/TaskArray.h
--- a/Engine/Source/Runtime/Core/Private/AutoRTFM/TaskArray.h
+++ b/Engine/Source/Runtime/Core/Private/AutoRTFM/TaskArray.h
@@ -181,6 +181,31 @@ public:
         return true;
     }
 
+	template<typename TFunc>
+	bool ForEachBackward(const TFunc& Func)
+	{
+		for (SKeyValuePair& EntryKVP : TBackwards(Latest))
+		{
+			if (!Func(EntryKVP.Val))
+			{
+				return false;
+			}
+		}
+
+		for (FInternalArray& StashedVectorBox : TBackwards(Stash))
+		{
+			for (SKeyValuePair& EntryKVP : TBackwards(StashedVectorBox))
+			{
+				if (!Func(EntryKVP.Val))
+				{
+					return false;
+				}
+			}
+		}
+
+		return true;
+	}
+
     void Reset()
     {
         Latest.Empty();
diff --git a/Engine/Source/Runtime/Core/Private/AutoRTFM/Transaction.cpp b/Engine/Source/Runtime/Core/Private/AutoRTFM/Transaction.cpp
--- a/Engine/Source/Runtime/Core/Private/AutoRTFM/Transaction.cpp
+++ b/Engine/Source/Runtime/Core/Private/AutoRTFM/Transaction.cpp
@@ -34,8 +34,21 @@ void FTransaction::AbortWithoutThrowing()
     Stats.Collect<EStatsKind::Abort>();
     CollectStats();
 
+    // Call the destructors of all the OnCommit functors before undoing the transactional memory and
+    // calling the OnAbort callbacks. This is important as the callback functions may have captured
+    // variables that are depending on the allocated memory. 
+    CommitTasks.Reset();
+
     Undo();
-	AbortTasks.ForEachBackward([&](const TFunction<void()>& Task) -> bool { Task(); return true; });
+	AbortTasks.ForEachBackward([&](TFunction<void()>& Task) -> bool 
+    { 
+        // Call and then reset each of the tasks in reverse order.
+        // This ensures that the task and its destructor are called in reverse chronological order,
+        // which is important if the function has captures with non-trivial destructors.
+        Task();
+        Task.Reset();
+        return true; 
+    });
 
     if (IsNested())
     {
