[AutoRTFM] Add unit test for function pointers into standard-library functions

Fortunately, this seemed to work as expected and didn't expose any lurking issues.

diff --git a/Engine/Source/Programs/AutoRTFMTests/Private/FunctionPointerTests.cpp b/Engine/Source/Programs/AutoRTFMTests/Private/FunctionPointerTests.cpp
--- a/Engine/Source/Programs/AutoRTFMTests/Private/FunctionPointerTests.cpp
+++ b/Engine/Source/Programs/AutoRTFMTests/Private/FunctionPointerTests.cpp
@@ -1,12 +1,14 @@
 // Copyright Epic Games, Inc. All Rights Reserved.
 
 #include "Catch2Includes.h"
-#include <AutoRTFM/AutoRTFM.h>
 #include "Templates/Function.h"
 
+#include <AutoRTFM/AutoRTFM.h>
+#include <cmath>
+
 namespace
 {
-	int Something()
+	int TestCFunction()
 	{
 		if (AutoRTFM::IsClosed())
 		{
@@ -20,14 +22,19 @@ namespace
 		return 43;
 	}
 
-	typedef int (*CStyleType)();
+	using CStyleType = int (*)();
+	using CosfType = float (*)(float);
 
 	UE_DISABLE_OPTIMIZATION_SHIP
-	CStyleType GetSomething()
+	CStyleType GetTestCFunction()
 	{
-		return &Something;
+		return &TestCFunction;
 	}
 
+	CosfType GetCosfFunction()
+	{
+		return &std::cosf;
+	}
 	UE_ENABLE_OPTIMIZATION_SHIP
 }
 
@@ -36,13 +43,41 @@ TEST_CASE("FunctionPointer.CStyle")
 	int Result = 0;
 	AutoRTFM::Commit([&]
 		{
-			CStyleType CStyle = GetSomething();
+			CStyleType CStyle = GetTestCFunction();
 			Result = CStyle();
 		});
 
 	REQUIRE(42 == Result);
 }
 
+TEST_CASE("FunctionPointer.StandardLibrary")
+{
+	SECTION("Created inside transaction")
+	{
+		float Result = 0;
+		AutoRTFM::Commit([&]
+			{
+				CosfType fn = GetCosfFunction();
+				Result = fn(0.0f);
+			});
+
+		REQUIRE(Result == 1.0f);
+	}
+
+	SECTION("Created outside transaction")
+	{
+		int Result = 0;
+		CosfType fn = GetCosfFunction();
+
+		AutoRTFM::Commit([&]
+			{
+				Result = fn(0.0f);
+			});
+
+		REQUIRE(Result == 1.0f);
+	}
+}
+
 TEST_CASE("FunctionPointer.TFunction")
 {
 	SECTION("Created inside transaction")
@@ -89,7 +124,6 @@ TEST_CASE("FunctionPointer.TFunction")
 			});
 
 		REQUIRE(42 == Result);
-
 	}
 }
 
