Add a `SingleThreaded` mode to `TQueue` and use it under AutoRTFM to fix a threading issue.

diff --git a/Engine/Source/Programs/AutoRTFMTests/Private/UECoreTests.cpp b/Engine/Source/Programs/AutoRTFMTests/Private/UECoreTests.cpp
--- a/Engine/Source/Programs/AutoRTFMTests/Private/UECoreTests.cpp
+++ b/Engine/Source/Programs/AutoRTFMTests/Private/UECoreTests.cpp
@@ -13,6 +13,7 @@
 #include "MyAutoRTFMTestObject.h"
 #include "Misc/TransactionallySafeScopeLock.h"
 #include "Misc/TransactionallySafeRWScopeLock.h"
+#include "Containers/Queue.h"
 
 TEST_CASE("UECore.FDelegateHandle")
 {
@@ -1511,3 +1512,226 @@ TEST_CASE("UECore.FString")
 		REQUIRE(String == "Foo 'Stuff' BAR");
 	}
 }
+
+TEST_CASE("UECore.TQueue")
+{
+	SECTION("SingleThreaded")
+	{
+		SECTION("Constructor")
+		{
+			AutoRTFM::Commit([&]
+				{
+					TQueue<int, EQueueMode::SingleThreaded> Queue;
+
+					AutoRTFM::Open([&]
+						{
+							REQUIRE(nullptr == Queue.Peek());
+						});
+				});
+		}
+
+		SECTION("Dequeue")
+		{
+			TQueue<int, EQueueMode::SingleThreaded> Queue;
+			REQUIRE(Queue.Enqueue(42));
+			REQUIRE(!Queue.IsEmpty());
+
+			int Value = 0;
+			bool bSucceeded = false;
+
+			AutoRTFM::ETransactionResult Result = AutoRTFM::Transact([&]
+				{
+					bSucceeded = Queue.Dequeue(Value);
+					AutoRTFM::AbortTransaction();
+				});
+
+			REQUIRE(AutoRTFM::ETransactionResult::AbortedByRequest == Result);
+			REQUIRE(!bSucceeded);
+			REQUIRE(0 == Value);
+			REQUIRE(42 == *Queue.Peek());
+
+			Result = AutoRTFM::Transact([&]
+				{
+					bSucceeded = Queue.Dequeue(Value);
+				});
+
+			REQUIRE(AutoRTFM::ETransactionResult::Committed == Result);
+			REQUIRE(bSucceeded);
+			REQUIRE(42 == Value);
+			REQUIRE(Queue.IsEmpty());
+		}
+
+		SECTION("Empty")
+		{
+			TQueue<int, EQueueMode::SingleThreaded> Queue;
+			REQUIRE(Queue.Enqueue(42));
+			REQUIRE(!Queue.IsEmpty());
+
+			AutoRTFM::ETransactionResult Result = AutoRTFM::Transact([&]
+				{
+					Queue.Empty();
+					AutoRTFM::Open([&] { REQUIRE(Queue.IsEmpty()); });
+					AutoRTFM::AbortTransaction();
+				});
+
+			REQUIRE(AutoRTFM::ETransactionResult::AbortedByRequest == Result);
+			REQUIRE(42 == *Queue.Peek());
+
+			Result = AutoRTFM::Transact([&]
+				{
+					Queue.Empty();
+				});
+
+			REQUIRE(AutoRTFM::ETransactionResult::Committed == Result);
+			REQUIRE(Queue.IsEmpty());
+		}
+
+		SECTION("Enqueue")
+		{
+			TQueue<int, EQueueMode::SingleThreaded> Queue;
+
+			bool bSucceeded = false;
+
+			AutoRTFM::ETransactionResult Result = AutoRTFM::Transact([&]
+				{
+					bSucceeded = Queue.Enqueue(42);
+					AutoRTFM::AbortTransaction();
+				});
+
+			REQUIRE(AutoRTFM::ETransactionResult::AbortedByRequest == Result);
+			REQUIRE(Queue.IsEmpty());
+			REQUIRE(!bSucceeded);
+
+			Result = AutoRTFM::Transact([&]
+				{
+					bSucceeded = Queue.Enqueue(42);
+				});
+
+			REQUIRE(AutoRTFM::ETransactionResult::Committed == Result);
+			REQUIRE(42 == *Queue.Peek());
+			REQUIRE(bSucceeded);
+		}
+
+		SECTION("IsEmpty")
+		{
+			TQueue<int, EQueueMode::SingleThreaded> Queue;
+			REQUIRE(Queue.IsEmpty());
+
+			bool bIsEmpty = false;
+
+			AutoRTFM::ETransactionResult Result = AutoRTFM::Transact([&]
+				{
+					bIsEmpty = Queue.IsEmpty();
+					AutoRTFM::AbortTransaction();
+				});
+
+			REQUIRE(AutoRTFM::ETransactionResult::AbortedByRequest == Result);
+			REQUIRE(!bIsEmpty);
+			
+			Result = AutoRTFM::Transact([&]
+				{
+					bIsEmpty = Queue.IsEmpty();
+				});
+
+			REQUIRE(AutoRTFM::ETransactionResult::Committed == Result);
+			REQUIRE(bIsEmpty);
+
+			Queue.Enqueue(42);
+			REQUIRE(!Queue.IsEmpty());
+
+			Result = AutoRTFM::Transact([&]
+				{
+					bIsEmpty = Queue.IsEmpty();
+					AutoRTFM::AbortTransaction();
+				});
+
+			REQUIRE(AutoRTFM::ETransactionResult::AbortedByRequest == Result);
+			REQUIRE(bIsEmpty);
+
+			Result = AutoRTFM::Transact([&]
+				{
+					bIsEmpty = Queue.IsEmpty();
+				});
+
+			REQUIRE(AutoRTFM::ETransactionResult::Committed == Result);
+			REQUIRE(!bIsEmpty);
+		}
+
+		SECTION("Peek")
+		{
+			TQueue<int, EQueueMode::SingleThreaded> Queue;
+			REQUIRE(Queue.Enqueue(42));
+
+			AutoRTFM::ETransactionResult Result = AutoRTFM::Transact([&]
+				{
+					*Queue.Peek() = 13;
+					AutoRTFM::AbortTransaction();
+				});
+
+			REQUIRE(AutoRTFM::ETransactionResult::AbortedByRequest == Result);
+			REQUIRE(42 == *Queue.Peek());
+
+			Result = AutoRTFM::Transact([&]
+				{
+					*Queue.Peek() = 13;
+				});
+
+			REQUIRE(AutoRTFM::ETransactionResult::Committed == Result);
+			REQUIRE(13 == *Queue.Peek());
+		}
+
+		SECTION("Pop")
+		{
+			SECTION("Empty")
+			{
+				TQueue<int, EQueueMode::SingleThreaded> Queue;
+
+				bool bSucceeded = true;
+
+				AutoRTFM::ETransactionResult Result = AutoRTFM::Transact([&]
+					{
+						bSucceeded = Queue.Pop();
+						AutoRTFM::AbortTransaction();
+					});
+
+				REQUIRE(AutoRTFM::ETransactionResult::AbortedByRequest == Result);
+				REQUIRE(bSucceeded);
+
+				Result = AutoRTFM::Transact([&]
+					{
+						bSucceeded = Queue.Pop();
+					});
+
+				REQUIRE(AutoRTFM::ETransactionResult::Committed == Result);
+				REQUIRE(!bSucceeded);
+			}
+
+			SECTION("Non Empty")
+			{
+				TQueue<int, EQueueMode::SingleThreaded> Queue;
+				REQUIRE(Queue.Enqueue(42));
+
+				bool bSucceeded = false;
+
+				AutoRTFM::ETransactionResult Result = AutoRTFM::Transact([&]
+					{
+						bSucceeded = Queue.Pop();
+						AutoRTFM::AbortTransaction();
+					});
+
+				REQUIRE(AutoRTFM::ETransactionResult::AbortedByRequest == Result);
+				REQUIRE(!bSucceeded);
+				REQUIRE(!Queue.IsEmpty());
+
+				Result = AutoRTFM::Transact([&]
+					{
+						bSucceeded = Queue.Pop();
+					});
+
+				REQUIRE(AutoRTFM::ETransactionResult::Committed == Result);
+				REQUIRE(bSucceeded);
+				REQUIRE(Queue.IsEmpty());
+			}
+		}
+	}
+}
diff --git a/Engine/Source/Runtime/Core/Public/Containers/Queue.h b/Engine/Source/Runtime/Core/Public/Containers/Queue.h
--- a/Engine/Source/Runtime/Core/Public/Containers/Queue.h
+++ b/Engine/Source/Runtime/Core/Public/Containers/Queue.h
@@ -18,7 +18,10 @@ enum class EQueueMode
 	Mpsc,
 
 	/** Single-producer, single-consumer queue. */
-	Spsc
+	Spsc,
+
+	/** Single-threaded - no guarantees of concurrent safety. */
+	SingleThreaded,
 };
 
 
@@ -77,8 +80,12 @@ public:
 		{
 			return false;
 		}
-		
-		TSAN_AFTER(&Tail->NextNode);
+
+		if constexpr (Mode != EQueueMode::SingleThreaded)
+		{
+			TSAN_AFTER(&Tail->NextNode);
+		}
+
 		OutItem = MoveTemp(Popped->Item);
 
 		TNode* OldTail = Tail;
@@ -129,8 +136,13 @@ public:
 		{
 			OldHead = Head;
 			Head = NewNode;
-			TSAN_BEFORE(&OldHead->NextNode);
-			FPlatformMisc::MemoryBarrier();
+
+			if constexpr (Mode == EQueueMode::Spsc)
+			{
+				TSAN_BEFORE(&OldHead->NextNode);
+				FPlatformMisc::MemoryBarrier();
+			}
+
             OldHead->NextNode = NewNode;
 		}
 
@@ -166,8 +178,13 @@ public:
 		{
 			OldHead = Head;
 			Head = NewNode;
-			TSAN_BEFORE(&OldHead->NextNode);
-			FPlatformMisc::MemoryBarrier();
+
+			if constexpr (Mode == EQueueMode::Spsc)
+			{
+				TSAN_BEFORE(&OldHead->NextNode);
+				FPlatformMisc::MemoryBarrier();
+			}
+
 			OldHead->NextNode = NewNode;
 		}
 
@@ -244,8 +261,11 @@ public:
 		{
 			return false;
 		}
-		
-		TSAN_AFTER(&Tail->NextNode);
+
+		if constexpr (Mode != EQueueMode::SingleThreaded)
+		{
+			TSAN_AFTER(&Tail->NextNode);
+		}
 
 		TNode* OldTail = Tail;
 		Tail = Popped;
