[AutoRTFM] Make UAIPerceptionSystem::UnregisterSource() transaction-safe

Add and use FRWTransactionallySafeAccessDetector, a transaction-safe version of FRWAccessDetector.
Use a FTransactionallySafeCriticalSection for UObjectDeleteListenersCritical.

diff --git a/Engine/Source/Programs/AutoRTFMTests/Private/MTTSAccessDetector.cpp b/Engine/Source/Programs/AutoRTFMTests/Private/MTTSAccessDetector.cpp
new file mode 100644
--- /dev/null
+++ b/Engine/Source/Programs/AutoRTFMTests/Private/MTTSAccessDetector.cpp
@@ -0,0 +1,450 @@
+// Copyright Epic Games, Inc. All Rights Reserved.
+
+#include "Catch2Includes.h"
+#include "Misc/MTTransactionallySafeAccessDetector.h"
+#include <AutoRTFM/AutoRTFM.h>
+
+#if ENABLE_MT_DETECTOR
+
+TEST_CASE("MTTransactionallySafeAccessDetector")
+{
+	FRWTransactionallySafeAccessDetector Detector;
+
+	SECTION("AcquireWriteAccess, ReleaseWriteAccess")
+	{
+		Detector.AcquireWriteAccess();
+		Detector.ReleaseWriteAccess();
+	}
+
+	SECTION("AcquireReadAccess, ReleaseReadAccess")
+	{
+		Detector.AcquireReadAccess();
+		Detector.ReleaseReadAccess();
+	}
+
+	SECTION("AcquireReadAccess, AcquireReadAccess, ReleaseReadAccess, ReleaseReadAccess")
+	{
+		Detector.AcquireReadAccess();
+		Detector.AcquireReadAccess();
+		Detector.ReleaseReadAccess();
+		Detector.ReleaseReadAccess();
+	}
+
+	SECTION("Transact(AcquireWriteAccess, ReleaseWriteAccess)")
+	{
+		AutoRTFM::ETransactionResult Result = AutoRTFM::Transact([&]
+		{
+			Detector.AcquireWriteAccess();
+			Detector.ReleaseWriteAccess();
+		});
+		REQUIRE(AutoRTFM::ETransactionResult::Committed == Result);
+	}
+
+	SECTION("Transact(AcquireWriteAccess, ReleaseWriteAccess, Abort)")
+	{
+		AutoRTFM::ETransactionResult Result = AutoRTFM::Transact([&]
+		{
+			Detector.AcquireWriteAccess();
+			Detector.ReleaseWriteAccess();
+			AutoRTFM::AbortTransaction();
+		});
+		REQUIRE(AutoRTFM::ETransactionResult::AbortedByRequest == Result);
+	}
+
+	SECTION("Transact(AcquireWriteAccess), ReleaseWriteAccess")
+	{
+		AutoRTFM::ETransactionResult Result = AutoRTFM::Transact([&]
+		{
+			Detector.AcquireWriteAccess();
+		});
+		REQUIRE(AutoRTFM::ETransactionResult::Committed == Result);
+
+		Detector.ReleaseWriteAccess();
+	}
+
+	SECTION("Transact(AcquireWriteAccess, Abort)")
+	{
+		AutoRTFM::ETransactionResult Result = AutoRTFM::Transact([&]
+		{
+			Detector.AcquireWriteAccess();
+			AutoRTFM::AbortTransaction();
+		});
+		REQUIRE(AutoRTFM::ETransactionResult::AbortedByRequest == Result);
+	}
+
+	SECTION("Transact(AcquireReadAccess, Abort)")
+	{
+		AutoRTFM::ETransactionResult Result = AutoRTFM::Transact([&]
+		{
+			Detector.AcquireReadAccess();
+			AutoRTFM::AbortTransaction();
+		});
+		REQUIRE(AutoRTFM::ETransactionResult::AbortedByRequest == Result);
+	}
+
+	SECTION("AcquireReadAccess, Transact(AcquireReadAccess, Abort), ReleaseReadAccess")
+	{
+		Detector.AcquireReadAccess();
+
+		AutoRTFM::ETransactionResult Result = AutoRTFM::Transact([&]
+		{
+			Detector.AcquireReadAccess();
+			AutoRTFM::AbortTransaction();
+		});
+		REQUIRE(AutoRTFM::ETransactionResult::AbortedByRequest == Result);
+
+		Detector.ReleaseReadAccess();
+	}
+
+	SECTION("AcquireWriteAccess, Transact(ReleaseWriteAccess)")
+	{
+		Detector.AcquireWriteAccess();
+
+		AutoRTFM::ETransactionResult Result = AutoRTFM::Transact([&]
+		{
+			Detector.ReleaseWriteAccess();
+		});
+		REQUIRE(AutoRTFM::ETransactionResult::Committed == Result);
+	}
+
+	SECTION("AcquireReadAccess, Transact(ReleaseReadAccess)")
+	{
+		Detector.AcquireReadAccess();
+
+		AutoRTFM::ETransactionResult Result = AutoRTFM::Transact([&]
+		{
+			Detector.ReleaseReadAccess();
+		});
+		REQUIRE(AutoRTFM::ETransactionResult::Committed == Result);
+	}
+
+	SECTION("AcquireWriteAccess, Transact(ReleaseWriteAccess, Abort), ReleaseWriteAccess")
+	{
+		Detector.AcquireWriteAccess();
+
+		AutoRTFM::ETransactionResult Result = AutoRTFM::Transact([&]
+		{
+			Detector.ReleaseWriteAccess();
+			AutoRTFM::AbortTransaction();
+		});
+		REQUIRE(AutoRTFM::ETransactionResult::AbortedByRequest == Result);
+
+		Detector.ReleaseWriteAccess();
+	}
+
+	SECTION("Transact(AcquireReadAccess, AcquireReadAccess, ReleaseReadAccess, ReleaseReadAccess)")
+	{
+		AutoRTFM::ETransactionResult Result = AutoRTFM::Transact([&]
+		{
+			Detector.AcquireReadAccess();
+			Detector.AcquireReadAccess();
+			Detector.ReleaseReadAccess();
+			Detector.ReleaseReadAccess();
+		});
+		REQUIRE(AutoRTFM::ETransactionResult::Committed == Result);
+	}
+
+	SECTION("Transact(AcquireReadAccess, AcquireReadAccess, ReleaseReadAccess, ReleaseReadAccess, Abort)")
+	{
+		AutoRTFM::ETransactionResult Result = AutoRTFM::Transact([&]
+		{
+			Detector.AcquireReadAccess();
+			Detector.AcquireReadAccess();
+			Detector.ReleaseReadAccess();
+			Detector.ReleaseReadAccess();
+			AutoRTFM::AbortTransaction();
+		});
+		REQUIRE(AutoRTFM::ETransactionResult::AbortedByRequest == Result);
+	}
+	
+	SECTION("Transact(AcquireWriteAccess, ReleaseWriteAccess, AcquireReadAccess)")
+	{
+		AutoRTFM::ETransactionResult Result = AutoRTFM::Transact([&]
+		{
+			Detector.AcquireWriteAccess();
+			Detector.ReleaseWriteAccess();
+			Detector.AcquireReadAccess();
+		});
+		REQUIRE(AutoRTFM::ETransactionResult::Committed == Result);
+		
+		Detector.ReleaseReadAccess();
+	}
+
+	SECTION("Transact(AcquireWriteAccess, ReleaseWriteAccess, AcquireReadAccess, Abort)")
+	{
+		AutoRTFM::ETransactionResult Result = AutoRTFM::Transact([&]
+		{
+			Detector.AcquireWriteAccess();
+			Detector.ReleaseWriteAccess();
+			Detector.AcquireReadAccess();
+			AutoRTFM::AbortTransaction();
+		});
+		REQUIRE(AutoRTFM::ETransactionResult::AbortedByRequest == Result);
+	}
+
+	SECTION("Transact(AcquireWriteAccess, ReleaseWriteAccess, AcquireReadAccess, ReleaseReadAccess)")
+	{
+		AutoRTFM::ETransactionResult Result = AutoRTFM::Transact([&]
+		{
+			Detector.AcquireWriteAccess();
+			Detector.ReleaseWriteAccess();
+			Detector.AcquireReadAccess();
+			Detector.ReleaseReadAccess();
+		});
+		REQUIRE(AutoRTFM::ETransactionResult::Committed == Result);
+	}
+
+	SECTION("Transact(AcquireWriteAccess, ReleaseWriteAccess, AcquireReadAccess, ReleaseReadAccess, Abort)")
+	{
+		AutoRTFM::ETransactionResult Result = AutoRTFM::Transact([&]
+		{
+			Detector.AcquireWriteAccess();
+			Detector.ReleaseWriteAccess();
+			Detector.AcquireReadAccess();
+			Detector.ReleaseReadAccess();
+			AutoRTFM::AbortTransaction();
+		});
+		REQUIRE(AutoRTFM::ETransactionResult::AbortedByRequest == Result);
+	}
+
+	SECTION("AcquireReadAccess, Transact(ReleaseReadAccess, AcquireWriteAccess), ReleaseWriteAccess")
+	{
+		Detector.AcquireReadAccess();
+
+		AutoRTFM::ETransactionResult Result = AutoRTFM::Transact([&]
+		{
+			Detector.ReleaseReadAccess();
+			Detector.AcquireWriteAccess();
+		});
+		REQUIRE(AutoRTFM::ETransactionResult::Committed == Result);
+
+		Detector.ReleaseWriteAccess();
+	}
+
+	SECTION("AcquireReadAccess, Transact(ReleaseReadAccess, AcquireWriteAccess, Abort), ReleaseReadAccess")
+	{
+		Detector.AcquireReadAccess();
+
+		AutoRTFM::ETransactionResult Result = AutoRTFM::Transact([&]
+		{
+			Detector.ReleaseReadAccess();
+			Detector.AcquireWriteAccess();
+			AutoRTFM::AbortTransaction();
+		});
+		REQUIRE(AutoRTFM::ETransactionResult::AbortedByRequest == Result);
+
+		Detector.ReleaseReadAccess();
+	}
+
+	SECTION("AcquireReadAccess, Transact(ReleaseReadAccess, AcquireWriteAccess, ReleaseWriteAccess)")
+	{
+		Detector.AcquireReadAccess();
+
+		AutoRTFM::ETransactionResult Result = AutoRTFM::Transact([&]
+		{
+			Detector.ReleaseReadAccess();
+			Detector.AcquireWriteAccess();
+			Detector.ReleaseWriteAccess();
+		});
+		REQUIRE(AutoRTFM::ETransactionResult::Committed == Result);
+	}
+
+	SECTION("AcquireReadAccess, Transact(ReleaseReadAccess, AcquireWriteAccess, ReleaseWriteAccess, Abort), ReleaseReadAccess")
+	{
+		Detector.AcquireReadAccess();
+
+		AutoRTFM::ETransactionResult Result = AutoRTFM::Transact([&]
+		{
+			Detector.ReleaseReadAccess();
+			Detector.AcquireWriteAccess();
+			Detector.ReleaseWriteAccess();
+			AutoRTFM::AbortTransaction();
+		});
+		REQUIRE(AutoRTFM::ETransactionResult::AbortedByRequest == Result);
+
+		Detector.ReleaseReadAccess();
+	}
+
+	{
+		// These tests use AutoRTFM::Open() blocks to acquire / release locks.
+		// Transaction retries will cause the lock counts to go out of sync, so
+		// we can only test these with retries disabled retries.
+		struct DisableRetriesScope
+		{
+			DisableRetriesScope() : OldState(AutoRTFM::ForTheRuntime::GetRetryTransaction())
+			{
+				AutoRTFM::ForTheRuntime::SetRetryTransaction(AutoRTFM::ForTheRuntime::NoRetry);
+			}
+			~DisableRetriesScope()
+			{
+				AutoRTFM::ForTheRuntime::SetRetryTransaction(OldState);
+			}
+			AutoRTFM::ForTheRuntime::EAutoRTFMRetryTransactionState OldState;
+		};
+		DisableRetriesScope DisableRetries;
+
+		SECTION("Transact(AcquireReadAccess, Open(AcquireReadAccess, ReleaseReadAccess), ReleaseReadAccess)")
+		{
+			AutoRTFM::ETransactionResult Result = AutoRTFM::Transact([&]
+			{
+				Detector.AcquireReadAccess();
+				AutoRTFM::Open([&]
+				{
+					Detector.AcquireReadAccess();
+					Detector.ReleaseReadAccess();
+				});
+				Detector.ReleaseReadAccess();
+			});
+			REQUIRE(AutoRTFM::ETransactionResult::Committed == Result);
+		}
+
+		SECTION("Transact(AcquireReadAccess, Open(AcquireReadAccess, ReleaseReadAccess), ReleaseReadAccess, Abort)")
+		{
+			AutoRTFM::ETransactionResult Result = AutoRTFM::Transact([&]
+			{
+				Detector.AcquireReadAccess();
+				AutoRTFM::Open([&]
+				{
+					Detector.AcquireReadAccess();
+					Detector.ReleaseReadAccess();
+				});
+				Detector.ReleaseReadAccess();
+				AutoRTFM::AbortTransaction();
+			});
+			REQUIRE(AutoRTFM::ETransactionResult::AbortedByRequest == Result);
+		}
+
+		SECTION("Transact(AcquireReadAccess, Open(AcquireReadAccess)), ReleaseReadAccess, ReleaseReadAccess")
+		{
+			AutoRTFM::ETransactionResult Result = AutoRTFM::Transact([&]
+			{
+				Detector.AcquireReadAccess();
+				AutoRTFM::Open([&]
+				{
+					Detector.AcquireReadAccess();
+				});
+			});
+			REQUIRE(AutoRTFM::ETransactionResult::Committed == Result);
+
+			Detector.ReleaseReadAccess();
+			Detector.ReleaseReadAccess();
+		}
+
+		SECTION("Transact(Open(AcquireReadAccess), ReleaseReadAccess, AcquireWriteAccess, Abort), ReleaseReadAccess")
+		{
+			AutoRTFM::ETransactionResult Result = AutoRTFM::Transact([&]
+			{
+				AutoRTFM::Open([&]
+				{
+					Detector.AcquireReadAccess();
+				});
+				Detector.ReleaseReadAccess();
+				Detector.AcquireWriteAccess();
+				AutoRTFM::AbortTransaction();
+			});
+			REQUIRE(AutoRTFM::ETransactionResult::AbortedByRequest == Result);
+
+			Detector.ReleaseReadAccess();
+		}
+
+		SECTION("Transact(AcquireReadAccess, Open(AcquireReadAccess), Abort), ReleaseReadAccess")
+		{
+			AutoRTFM::ETransactionResult Result = AutoRTFM::Transact([&]
+			{
+				Detector.AcquireReadAccess();
+				AutoRTFM::Open([&]
+				{
+					Detector.AcquireReadAccess();
+				});
+				AutoRTFM::AbortTransaction();
+			});
+			REQUIRE(AutoRTFM::ETransactionResult::AbortedByRequest == Result);
+
+			Detector.ReleaseReadAccess();
+		}
+
+		SECTION("AcquireReadAccess, Transact(AcquireReadAccess, Open(ReleaseReadAccess), Abort)")
+		{
+			Detector.AcquireReadAccess();
+
+			AutoRTFM::ETransactionResult Result = AutoRTFM::Transact([&]
+			{
+				Detector.AcquireReadAccess();
+				AutoRTFM::Open([&]
+				{
+					Detector.ReleaseReadAccess();
+				});
+				AutoRTFM::AbortTransaction();
+			});
+			REQUIRE(AutoRTFM::ETransactionResult::AbortedByRequest == Result);
+		}
+
+		SECTION("AcquireWriteAccess, Transact(Open(ReleaseWriteAccess), AcquireReadAccess, ReleaseReadAccess)")
+		{
+			Detector.AcquireWriteAccess();
+
+			AutoRTFM::ETransactionResult Result = AutoRTFM::Transact([&]
+			{
+				AutoRTFM::Open([&]
+				{
+					Detector.ReleaseWriteAccess();
+				});
+				Detector.AcquireReadAccess();
+			});
+			REQUIRE(AutoRTFM::ETransactionResult::Committed == Result);
+
+			Detector.ReleaseReadAccess();
+		}
+
+		SECTION("AcquireWriteAccess, Transact(Open(ReleaseWriteAccess), AcquireReadAccess, Abort)")
+		{
+			Detector.AcquireWriteAccess();
+
+			AutoRTFM::ETransactionResult Result = AutoRTFM::Transact([&]
+			{
+				AutoRTFM::Open([&]
+				{
+					Detector.ReleaseWriteAccess();
+				});
+				Detector.AcquireReadAccess();
+				AutoRTFM::AbortTransaction();
+			});
+			REQUIRE(AutoRTFM::ETransactionResult::AbortedByRequest == Result);
+		}
+
+		SECTION("AcquireReadAccess, Transact(Open(ReleaseReadAccess), AcquireWriteAccess), ReleaseWriteAccess")
+		{
+			Detector.AcquireReadAccess();
+
+			AutoRTFM::ETransactionResult Result = AutoRTFM::Transact([&]
+			{
+				AutoRTFM::Open([&]
+				{
+					Detector.ReleaseReadAccess();
+				});
+				Detector.AcquireWriteAccess();
+			});
+			REQUIRE(AutoRTFM::ETransactionResult::Committed == Result);
+		
+			Detector.ReleaseWriteAccess();
+		}
+
+		SECTION("AcquireReadAccess, Transact(Open(ReleaseReadAccess), AcquireWriteAccess, Abort)")
+		{
+			Detector.AcquireReadAccess();
+
+			AutoRTFM::ETransactionResult Result = AutoRTFM::Transact([&]
+			{
+				AutoRTFM::Open([&]
+				{
+					Detector.ReleaseReadAccess();
+				});
+				Detector.AcquireWriteAccess();
+				AutoRTFM::AbortTransaction();
+			});
+			REQUIRE(AutoRTFM::ETransactionResult::AbortedByRequest == Result);
+		}
+	}
+}
+
+#endif // ENABLE_MT_DETECTOR
diff --git a/Engine/Source/Runtime/AIModule/Classes/Perception/AISense_Sight.h b/Engine/Source/Runtime/AIModule/Classes/Perception/AISense_Sight.h
--- a/Engine/Source/Runtime/AIModule/Classes/Perception/AISense_Sight.h
+++ b/Engine/Source/Runtime/AIModule/Classes/Perception/AISense_Sight.h
@@ -6,7 +6,7 @@
 #include "UObject/ObjectMacros.h"
 #include "GenericTeamAgentInterface.h"
 #include "WorldCollision.h"
-#include "Misc/MTAccessDetector.h"
+#include "Misc/MTTransactionallySafeAccessDetector.h"
 #include "Perception/AISense.h"
 #include "AISense_Sight.generated.h"
 
@@ -278,7 +278,7 @@ protected:
 	FOnPendingVisibilityQueryProcessedDelegate OnPendingCanBeSeenQueryProcessedDelegate;
 	FTraceDelegate OnPendingTraceQueryProcessedDelegate;
 
-	UE_MT_DECLARE_RW_ACCESS_DETECTOR(QueriesListAccessDetector);
+	UE_MT_DECLARE_TS_RW_ACCESS_DETECTOR(QueriesListAccessDetector);
 
 public:
 
diff --git a/Engine/Source/Runtime/Core/Private/AutoRTFM/API.cpp b/Engine/Source/Runtime/Core/Private/AutoRTFM/API.cpp
--- a/Engine/Source/Runtime/Core/Private/AutoRTFM/API.cpp
+++ b/Engine/Source/Runtime/Core/Private/AutoRTFM/API.cpp
@@ -145,6 +145,15 @@ namespace AutoRTFM
 #endif
 		}
 
+		EAutoRTFMRetryTransactionState GetRetryTransaction()
+		{
+#if UE_AUTORTFM
+			return static_cast<EAutoRTFMRetryTransactionState>(GAutoRTFMRetryTransactions);
+#else
+			return NoRetry;
+#endif
+		}
+
 		bool ShouldRetryNonNestedTransactions()
 		{
 #if UE_AUTORTFM
diff --git a/Engine/Source/Runtime/Core/Public/AutoRTFM/AutoRTFM.h b/Engine/Source/Runtime/Core/Public/AutoRTFM/AutoRTFM.h
--- a/Engine/Source/Runtime/Core/Public/AutoRTFM/AutoRTFM.h
+++ b/Engine/Source/Runtime/Core/Public/AutoRTFM/AutoRTFM.h
@@ -794,6 +794,9 @@ namespace ForTheRuntime
 	// Set whether we should retry transactions.
 	UE_AUTORTFM_API void SetRetryTransaction(EAutoRTFMRetryTransactionState State);
 
+	// Returns whether we should retry transactions.
+	UE_AUTORTFM_API EAutoRTFMRetryTransactionState GetRetryTransaction();
+
 	// Returns true if we should retry non-nested transactions.
 	UE_AUTORTFM_API bool ShouldRetryNonNestedTransactions();
 
diff --git a/Engine/Source/Runtime/Core/Public/Misc/MTTransactionallySafeAccessDetector.h b/Engine/Source/Runtime/Core/Public/Misc/MTTransactionallySafeAccessDetector.h
new file mode 100644
--- /dev/null
+++ b/Engine/Source/Runtime/Core/Public/Misc/MTTransactionallySafeAccessDetector.h
@@ -0,0 +1,299 @@
+// Copyright Epic Games, Inc. All Rights Reserved.
+
+#pragma once
+
+#include "AutoRTFM/AutoRTFM.h"
+#include "HAL/Platform.h"
+#include "MTAccessDetector.h"
+#include "Misc/AssertionMacros.h"
+
+#if ENABLE_MT_DETECTOR
+
+/**
+ * A version of FRWAccessDetector which also can be used in AutoRTFM transactions.
+ *
+ * When methods are called outside of a closed transaction they will behave
+ * exactly the same as FRWAccessDetector.
+ *
+ * When called from a closed transaction:
+ *  - The first call to AcquireWriteLock() will hold the internal write lock for
+ *    the remaining duration of the transaction, even if the transaction calls
+ *    ReleaseWriteLock(). This is done to guard against writes that would occur
+ *    if the transaction were aborted. The transaction's use of the API remains
+ *    unchanged, but unlocks and read locks are effectively no-ops and other 
+ *    threads will "see" a longer duration of the write lock.
+ *  - Each call to [Acquire|Release][Read|Write]Access is recorded as a pair
+ *    of signed read and write counters so that the OnAbort and OnCommit
+ *    handlers can restore or apply the read and write lock state, respectively.
+ */
+struct FRWTransactionallySafeAccessDetectorDefinition
+{
+public:
+	/**
+	 * Acquires read access, will check if there are any writers
+	 * @return true if no errors were detected
+	 */
+	FORCEINLINE bool AcquireReadAccess() const
+	{
+		if (!AutoRTFM::IsClosed())
+		{
+			// Called outside of a closed transaction.
+			// Forward to the inner detector and return.
+			return Inner.AcquireReadAccess();
+		}
+
+		// AcquireReadAccess() called inside a closed transaction.
+		return AutoRTFM::Open([this]
+		{
+			if (!TransactionalState.HoldsInternalWriteLock)
+			{
+				// AcquireReadAccess() called before the first transactional
+				// call to AcquireWriteAccess(). Call the inner detector.
+				if (UNLIKELY(!Inner.AcquireReadAccess()))
+				{
+					return false;
+				}
+
+				// Register the abort handler if this is the first transactional
+				// method call.
+				MaybeRegisterAbortHandler();
+			}
+
+			// Increment the transactional read-lock counter.
+			TransactionalState.ReadLockDelta++;
+			return true;
+		});
+	}
+
+	/**
+	 * Releases read access, will check if there are any writers
+	 * @return true if no errors were detected
+	 */
+	FORCEINLINE bool ReleaseReadAccess() const
+	{
+		if (!AutoRTFM::IsClosed())
+		{
+			// Called outside of a closed transaction.
+			// Forward to the inner detector and return.
+			return Inner.ReleaseReadAccess();
+		}
+
+		// ReleaseReadAccess() called inside a closed transaction.
+		return AutoRTFM::Open([this]
+		{
+			if (!TransactionalState.HoldsInternalWriteLock)
+			{
+				// ReleaseReadAccess() called before the first transactional
+				// call to AcquireWriteAccess(). Call the inner detector.
+				if (UNLIKELY(!Inner.ReleaseReadAccess()))
+				{
+					return false;
+				}
+
+				// Register the abort handler if this is the first transactional
+				// method call.
+				MaybeRegisterAbortHandler();
+			}
+
+			// Decrement the transactional read-lock counter.
+			TransactionalState.ReadLockDelta--;
+			return true;
+		});
+	}
+
+	/** 
+	 * Acquires write access, will check if there are readers or other writers
+	 * @return true if no errors were detected
+	 */
+	FORCEINLINE bool AcquireWriteAccess() const
+	{
+		if (!AutoRTFM::IsClosed())
+		{
+			// Called outside of a closed transaction.
+			// Forward to the inner detector and return.
+			return Inner.AcquireWriteAccess();
+		}
+
+		return AutoRTFM::Open([this]
+		{
+			if (!TransactionalState.HoldsInternalWriteLock)
+			{
+				// First transactional call to AcquireWriteAccess(). 
+				// Call the inner detector.
+				if (UNLIKELY(!Inner.AcquireWriteAccess()))
+				{
+					return false;
+				}
+
+				// Register the abort handler if this is the first transactional
+				// method call.
+				MaybeRegisterAbortHandler();
+				// Register the commit handler to unlock the write-lock which
+				// will be held until the transaction is done.
+				RegisterCommitHandler();
+				TransactionalState.HoldsInternalWriteLock = true;
+			}
+
+			// Increment the transactional write-lock counter.
+			TransactionalState.WriteLockDelta++;
+			return true;
+		});
+	}
+
+	/** 
+	 * Releases write access, will check if there are readers or other writers
+	 * @return true if no errors were detected
+	 */
+	FORCEINLINE bool ReleaseWriteAccess() const
+	{
+		if (!AutoRTFM::IsClosed())
+		{
+			// Called outside of a closed transaction.
+			// Forward to the inner detector and return.
+			return Inner.ReleaseWriteAccess();
+		}
+
+		return AutoRTFM::Open([this]
+		{
+			if (!TransactionalState.HoldsInternalWriteLock)
+			{
+				// ReleaseWriteAccess() called before the first transactional
+				// call to AcquireWriteAccess(). Call the inner detector.
+				if (UNLIKELY(!Inner.ReleaseWriteAccess()))
+				{
+					return false;
+				}
+
+				// Register the abort handler if this is the first transactional
+				// method call.
+				MaybeRegisterAbortHandler();
+			}
+
+			// Decrement the transactional write-lock counter.
+			TransactionalState.WriteLockDelta--;
+			return true;
+		});
+	}
+
+private:
+	// Registers a transaction abort handler if this is the first call for the
+	// current transaction.
+	void MaybeRegisterAbortHandler() const
+	{
+		if (TransactionalState.AbortHandlerRegistered)
+		{
+			return; // Already registered.
+		}
+		TransactionalState.AbortHandlerRegistered = true;
+
+		// Call OnAbort() in the closed, otherwise it'll be a no-op.
+		(void) AutoRTFM::Close([&]
+		{
+			AutoRTFM::OnAbort([this]
+			{
+				// Transaction is being aborted. Undo all state changes.
+				if (TransactionalState.HoldsInternalWriteLock)
+				{
+					// AcquireWriteAccess() was called in the transaction.
+					// The initial state when entering the transaction must have
+					// been write-unlocked, otherwise we'd double-write-lock.
+					// Unlock the write.
+					Inner.ReleaseWriteAccess();
+				}
+				else
+				{
+					// AcquireWriteAccess() was not called in the transaction.
+					// Undo any transactional call ReleaseWriteAccess().
+					if (TransactionalState.WriteLockDelta < 0)
+					{
+						Inner.AcquireWriteAccess();
+					}
+					// Undo any transactional calls AcquireReadAccess().
+					for (; TransactionalState.ReadLockDelta > 0; TransactionalState.ReadLockDelta--)
+					{
+						Inner.ReleaseReadAccess();
+					}
+				}
+				
+				// Undo any transactional calls ReleaseReadAccess().
+				for (; TransactionalState.ReadLockDelta < 0; TransactionalState.ReadLockDelta++)
+				{
+					Inner.AcquireReadAccess();
+				}
+
+				// Clear the TransactionalState for the next transaction.
+				TransactionalState = {};
+			});
+		});
+	}
+
+	// Registers a transaction commit handler to rebalance locks after the first
+	// call to AcquireWriteAccess().
+	void RegisterCommitHandler() const
+	{
+		// Call OnCommit() in the closed, otherwise it'll invoke the callback 
+		// immediately instead of when the transaction is committed. 
+		(void) AutoRTFM::Close([&]
+		{
+			AutoRTFM::OnCommit([this]
+			{
+				// Transaction is being committed.
+				check(TransactionalState.HoldsInternalWriteLock);
+				check(TransactionalState.WriteLockDelta >= 0);
+				// Release the write-lock, if the write-lock is now balanced.
+				if (TransactionalState.WriteLockDelta == 0)
+				{
+					Inner.ReleaseWriteAccess();
+				}
+				// Calls to AcquireReadAccess() with the write-lock held 
+				// were a no-op, so apply them now.
+				for (; TransactionalState.ReadLockDelta > 0; TransactionalState.ReadLockDelta--)
+				{
+					Inner.AcquireReadAccess();
+				}
+
+				// Clear the TransactionalState for the next transaction.
+				TransactionalState = {};
+			});
+		});
+	}
+
+	struct FTransactionalState
+	{
+		// Incremented for each transactional call to AcquireReadAccess(),
+		// decremented for each transactional call to ReleaseReadAccess().
+		int16 ReadLockDelta = 0;
+
+		// Incremented for each transactional call to AcquireWriteAccess(),
+		// decremented for each transactional call to ReleaseWriteAccess().
+		int16 WriteLockDelta = 0;
+		
+		// If true, then the write-lock has been taken in a transaction and will
+		// be held until the transaction is complete (either aborted or 
+		// committed).
+		bool HoldsInternalWriteLock = false;
+
+		// True if the abort handler has been registered.
+		bool AbortHandlerRegistered = false;
+	};
+
+	// The inner FRWAccessDetector.
+	FRWAccessDetector Inner;
+
+	// The state held for calls made when in a transaction.
+	mutable FTransactionalState TransactionalState;
+};
+
+#if UE_AUTORTFM
+using FRWTransactionallySafeAccessDetector = FRWTransactionallySafeAccessDetectorDefinition;
+#else
+using FRWTransactionallySafeAccessDetector = FRWAccessDetector;
+#endif
+
+#define UE_MT_DECLARE_TS_RW_ACCESS_DETECTOR(AccessDetector) FRWTransactionallySafeAccessDetector AccessDetector;
+
+#else // ENABLE_MT_DETECTOR
+
+#define UE_MT_DECLARE_TS_RW_ACCESS_DETECTOR(AccessDetector)
+
+#endif // ENABLE_MT_DETECTOR
diff --git a/Engine/Source/Runtime/CoreUObject/Private/UObject/UObjectArray.cpp b/Engine/Source/Runtime/CoreUObject/Private/UObject/UObjectArray.cpp
--- a/Engine/Source/Runtime/CoreUObject/Private/UObject/UObjectArray.cpp
+++ b/Engine/Source/Runtime/CoreUObject/Private/UObject/UObjectArray.cpp
@@ -8,6 +8,7 @@
 #include "HAL/IConsoleManager.h"
 #include "HAL/LowLevelMemStats.h"
 #include "Misc/ScopeLock.h"
+#include "Misc/TransactionallySafeScopeLock.h"
 #include "ProfilingDebugging/AssetMetadataTrace.h"
 #include "UObject/UObjectAllocator.h"
 #include "UObject/Class.h"
@@ -313,7 +314,7 @@ void FUObjectArray::ResetSerialNumber(UObjectBase* Object)
 void FUObjectArray::RemoveObjectFromDeleteListeners(UObjectBase* Object)
 {
 #if THREADSAFE_UOBJECTS
-	FScopeLock UObjectDeleteListenersLock(&UObjectDeleteListenersCritical);
+	FTransactionallySafeScopeLock UObjectDeleteListenersLock(&UObjectDeleteListenersCritical);
 #endif
 	int32 Index = Object->InternalIndex;
 	check(Index >= 0);
@@ -400,7 +401,7 @@ void FUObjectArray::RemoveUObjectCreateListener(FUObjectCreateListener* Listener
 void FUObjectArray::AddUObjectDeleteListener(FUObjectDeleteListener* Listener)
 {
 #if THREADSAFE_UOBJECTS
-	FScopeLock UObjectDeleteListenersLock(&UObjectDeleteListenersCritical);
+	FTransactionallySafeScopeLock UObjectDeleteListenersLock(&UObjectDeleteListenersCritical);
 #endif
 	check(!UObjectDeleteListeners.Contains(Listener));
 	UObjectDeleteListeners.Add(Listener);
@@ -414,7 +415,7 @@ void FUObjectArray::AddUObjectDeleteListener(FUObjectDeleteListener* Listener)
 void FUObjectArray::RemoveUObjectDeleteListener(FUObjectDeleteListener* Listener)
 {
 #if THREADSAFE_UOBJECTS
-	FScopeLock UObjectDeleteListenersLock(&UObjectDeleteListenersCritical);
+	FTransactionallySafeScopeLock UObjectDeleteListenersLock(&UObjectDeleteListenersCritical);
 #endif
 	UObjectDeleteListeners.RemoveSingleSwap(Listener);
 }
@@ -486,7 +487,7 @@ void FUObjectArray::ShutdownUObjectArray()
 {
 	{
 #if THREADSAFE_UOBJECTS
-		FScopeLock UObjectDeleteListenersLock(&UObjectDeleteListenersCritical);
+		FTransactionallySafeScopeLock UObjectDeleteListenersLock(&UObjectDeleteListenersCritical);
 #endif
 		for (int32 Index = UObjectDeleteListeners.Num() - 1; Index >= 0; --Index)
 		{
diff --git a/Engine/Source/Runtime/CoreUObject/Public/UObject/UObjectArray.h b/Engine/Source/Runtime/CoreUObject/Public/UObject/UObjectArray.h
--- a/Engine/Source/Runtime/CoreUObject/Public/UObject/UObjectArray.h
+++ b/Engine/Source/Runtime/CoreUObject/Public/UObject/UObjectArray.h
@@ -8,6 +8,7 @@
 
 #include "HAL/ThreadSafeCounter.h"
 #include "Containers/LockFreeList.h"
+#include "Misc/TransactionallySafeCriticalSection.h"
 #include "UObject/GarbageCollectionGlobals.h"
 #include "UObject/UObjectBase.h"
 
@@ -1316,7 +1317,7 @@ private:
 	 */
 	TArray<FUObjectDeleteListener* > UObjectDeleteListeners;
 #if THREADSAFE_UOBJECTS
-	mutable FCriticalSection UObjectDeleteListenersCritical;
+	mutable FTransactionallySafeCriticalSection UObjectDeleteListenersCritical;
 #endif
 
 	/** Current primary serial number **/
@@ -1342,7 +1343,7 @@ public:
     {
 		FScopeLock ObjListLock(&ObjObjectsCritical);
 #if THREADSAFE_UOBJECTS
-		FScopeLock ListenersLock(&UObjectDeleteListenersCritical);
+		FTransactionallySafeScopeLock ListenersLock(&UObjectDeleteListenersCritical);
 #endif
         return ObjObjects.GetAllocatedSize() + ObjAvailableList.GetAllocatedSize() + UObjectCreateListeners.GetAllocatedSize() + UObjectDeleteListeners.GetAllocatedSize();
     }
@@ -1350,7 +1351,7 @@ public:
 	SIZE_T GetDeleteListenersAllocatedSize(int32* OutNumListeners = nullptr) const
 	{
 #if THREADSAFE_UOBJECTS
-		FScopeLock ListenersLock(&UObjectDeleteListenersCritical);
+		FTransactionallySafeScopeLock ListenersLock(&UObjectDeleteListenersCritical);
 #endif
 		SIZE_T AllocatedSize = 0;
 		for (FUObjectDeleteListener* Listener : UObjectDeleteListeners)
