[AutoRTFM] Add PopAllOnAbortHandlers method.

This API deletes all pushed handlers for a given key. It is functionally equivalent to calling `PopOnAbortHandler` repeatedly, but more efficient.

diff --git a/Engine/Source/Programs/AutoRTFMTests/Private/Abort.cpp b/Engine/Source/Programs/AutoRTFMTests/Private/Abort.cpp
--- a/Engine/Source/Programs/AutoRTFMTests/Private/Abort.cpp
+++ b/Engine/Source/Programs/AutoRTFMTests/Private/Abort.cpp
@@ -250,13 +250,29 @@ TEST_CASE("Abort.PushOnAbortHandler_WithPop_NoAbort")
 {
 	int Value = 55;
 
+	AutoRTFM::Commit([&]
+		{
+			Value = 66;
+			AutoRTFM::PushOnAbortHandler(UIntToPointer(747), [&Value]() { Value = 77; });
+			Value = 88;
+
+			AutoRTFM::PopOnAbortHandler(UIntToPointer(747));
+		});
+
+	REQUIRE(Value == 88);
+}
+
+TEST_CASE("Abort.PushOnAbortHandler_WithPopAll_NoAbort")
+{
+	int Value = 55;
+
 	AutoRTFM::Commit([&]
 	{
 		Value = 66;
 		AutoRTFM::PushOnAbortHandler(UIntToPointer(747), [&Value](){ Value = 77; });
 		Value = 88;
 
-		AutoRTFM::PopOnAbortHandler(UIntToPointer(747));
+		AutoRTFM::PopAllOnAbortHandlers(UIntToPointer(747));
 	});
 
 	REQUIRE(Value == 88);
@@ -281,28 +297,69 @@ TEST_CASE("Abort.PushOnAbortHandler_WithPop_WithAbort")
 	REQUIRE(Value == 55);
 }
 
-TEST_CASE("Abort.PushOnAbortHandler_Duplicates1")
+TEST_CASE("Abort.PushOnAbortHandler_WithPopAll_WithAbort")
 {
 	int Value = 55;
 
 	const AutoRTFM::ETransactionResult Result = AutoRTFM::Transact([&]
 	{
 		Value = 66;
-		AutoRTFM::PushOnAbortHandler(UIntToPointer(747), [&Value](){ Value = 77; });
-		AutoRTFM::PushOnAbortHandler(UIntToPointer(747), [&Value](){ Value = 88; });
-		Value = 99;
+		AutoRTFM::PushOnAbortHandler(UIntToPointer(747), [&Value]() { Value = 77; });
+		Value = 88;
 
-		AutoRTFM::PopOnAbortHandler(UIntToPointer(747));
+		AutoRTFM::PopAllOnAbortHandlers(UIntToPointer(747));
 
 		AutoRTFM::AbortTransaction();
 	});
 
+	REQUIRE(AutoRTFM::ETransactionResult::AbortedByRequest == Result);
+	REQUIRE(Value == 55);
+}
+
+TEST_CASE("Abort.PushOnAbortHandler_Duplicates1")
+{
+	int Value = 55;
+
+	const AutoRTFM::ETransactionResult Result = AutoRTFM::Transact([&]
+		{
+			Value = 66;
+			AutoRTFM::PushOnAbortHandler(UIntToPointer(747), [&Value]() { Value = 77; });
+			AutoRTFM::PushOnAbortHandler(UIntToPointer(747), [&Value]() { Value = 88; });
+			Value = 99;
+
+			AutoRTFM::PopOnAbortHandler(UIntToPointer(747));
+
+			AutoRTFM::AbortTransaction();
+		});
+
 	REQUIRE(AutoRTFM::ETransactionResult::AbortedByRequest == Result);
 
 	// The first push on abort will still go through.
 	REQUIRE(Value == 77);
 }
 
+TEST_CASE("Abort.PushOnAbortHandler_PopAll_Duplicates")
+{
+	int Value = 55;
+
+	const AutoRTFM::ETransactionResult Result = AutoRTFM::Transact([&]
+	{
+		Value = 66;
+		AutoRTFM::PushOnAbortHandler(UIntToPointer(747), [&Value]() { Value = 77; });
+		AutoRTFM::PushOnAbortHandler(UIntToPointer(747), [&Value]() { Value = 88; });
+		Value = 99;
+
+		AutoRTFM::PopAllOnAbortHandlers(UIntToPointer(747));
+
+		AutoRTFM::AbortTransaction();
+	});
+
+	REQUIRE(AutoRTFM::ETransactionResult::AbortedByRequest == Result);
+
+	// No abort handlers should execute.
+	REQUIRE(Value == 55);
+}
+
 TEST_CASE("Abort.PushOnAbortHandler_Duplicates2")
 {
 	int Value = 55;
@@ -431,7 +488,7 @@ TEST_CASE("Abort.PushOnAbortHandler_Order")
 			const AutoRTFM::ETransactionResult Result = AutoRTFM::Transact([&]
 				{
 					AutoRTFM::OnAbort([&Value] { REQUIRE(40 == Value); Value += 1; });
-					
+
 					// Make a child transaction.
 					AutoRTFM::Commit([&]
 						{
diff --git a/Engine/Source/Runtime/Core/Private/AutoRTFM/TaskArray.h b/Engine/Source/Runtime/Core/Private/AutoRTFM/TaskArray.h
--- a/Engine/Source/Runtime/Core/Private/AutoRTFM/TaskArray.h
+++ b/Engine/Source/Runtime/Core/Private/AutoRTFM/TaskArray.h
@@ -71,31 +71,19 @@ public:
 
 		auto EraseLastAddedKeyFromArray = [Key](FInternalArray& Array) -> bool
 		{
-			for(typename FInternalArray::SizeType Idx = 0; Idx < Array.Num(); Idx++)
+			typename FInternalArray::SizeType LastIdx = Array.FindLastByPredicate([Key](const SKeyValuePair& Pair)
 			{
-				typename FInternalArray::SizeType BackwardsIdx = Array.Num() - 1 - Idx;
+				return Pair.Key == Key;
+			});
 
-				if(Array[BackwardsIdx].Key != Key)
-				{
-					continue;
-				}
-
-				// We found our key to erase! Nuke it.
-
-				// We start at the element just after our element.
-				BackwardsIdx++;
-
-				for (; BackwardsIdx < Array.Num(); BackwardsIdx++)
-				{
-					Array[BackwardsIdx - 1] = MoveTemp(Array[BackwardsIdx]);
-				}
-
-				Array.SetNum(Array.Num() - 1, EAllowShrinking::No);
-
-				return true;
+			if (LastIdx == INDEX_NONE)
+			{
+				return false;
 			}
 
-			return false;
+			// We found our key to erase! Nuke it.
+			Array.RemoveAt(LastIdx);
+			return true;
 		};
 
 		if (EraseLastAddedKeyFromArray(Latest))
@@ -114,6 +102,28 @@ public:
 		return false;
 	}
 
+	bool DeleteAllMatchingKeys(SKey Key)
+	{
+		ASSERT(Key != SKey{});
+
+		auto EraseKeyFromArray = [Key](FInternalArray& Array) -> size_t
+		{
+			return Array.RemoveAll([Key](const SKeyValuePair& Pair)
+			{
+				return Pair.Key == Key;
+			});
+		};
+
+		size_t NumErased = EraseKeyFromArray(Latest);
+
+		for (FInternalArray& StashedVectorBox : Stash)
+		{
+			NumErased += EraseKeyFromArray(StashedVectorBox);
+		}
+
+		return NumErased > 0;
+	}
+
     void AddAll(TTaskArray&& Other)
     {
         Canonicalize();
diff --git a/Engine/Source/Runtime/Core/Private/AutoRTFM/Transaction.h b/Engine/Source/Runtime/Core/Private/AutoRTFM/Transaction.h
--- a/Engine/Source/Runtime/Core/Private/AutoRTFM/Transaction.h
+++ b/Engine/Source/Runtime/Core/Private/AutoRTFM/Transaction.h
@@ -64,7 +64,8 @@ public:
     void DeferUntilCommit(TFunction<void()>&&);
     void DeferUntilAbort(TFunction<void()>&&);
     void PushDeferUntilAbortHandler(const void* Key, TFunction<void()>&&);
-    bool PopDeferUntilAbortHandler(const void* Key);
+	bool PopDeferUntilAbortHandler(const void* Key);
+	bool PopAllDeferUntilAbortHandlers(const void* Key);
 
     [[noreturn]] void AbortAndThrow();
     void AbortWithoutThrowing();
diff --git a/Engine/Source/Runtime/Core/Private/AutoRTFM/TransactionInlines.h b/Engine/Source/Runtime/Core/Private/AutoRTFM/TransactionInlines.h
--- a/Engine/Source/Runtime/Core/Private/AutoRTFM/TransactionInlines.h
+++ b/Engine/Source/Runtime/Core/Private/AutoRTFM/TransactionInlines.h
@@ -160,7 +160,12 @@ UE_AUTORTFM_FORCEINLINE void FTransaction::PushDeferUntilAbortHandler(const void
 
 UE_AUTORTFM_FORCEINLINE bool FTransaction::PopDeferUntilAbortHandler(const void* Key)
 {
-    return AbortTasks.DeleteKey(Key);
+	return AbortTasks.DeleteKey(Key);
+}
+
+UE_AUTORTFM_FORCEINLINE bool FTransaction::PopAllDeferUntilAbortHandlers(const void* Key)
+{
+	return AbortTasks.DeleteAllMatchingKeys(Key);
 }
 
 UE_AUTORTFM_FORCEINLINE void FTransaction::CollectStats() const
