[AutoRTFM] Expose IsInnerTransactionStack as public API.

This can be necessary when writing an OnAbort handler. We want to roll back changes to heap objects, but it's unsafe to roll back changes to objects on the inner-transaction's stack, so we need to be able to differentiate the two.

diff --git a/Engine/Source/Programs/AutoRTFMTests/Private/API.cpp b/Engine/Source/Programs/AutoRTFMTests/Private/API.cpp
--- a/Engine/Source/Programs/AutoRTFMTests/Private/API.cpp
+++ b/Engine/Source/Programs/AutoRTFMTests/Private/API.cpp
@@ -949,3 +949,34 @@ TEST_CASE("API.CheckConsistencyAssumingNoRaces")
         AutoRTFM::ForTheRuntime::CheckConsistencyAssumingNoRaces();
     });
 }
+
+TEST_CASE("API.IsInnerTransactionStack")
+{
+	{
+		int OnStackNotInTransaction = 1;
+		REQUIRE(!AutoRTFM::IsInnerTransactionStack(&OnStackNotInTransaction));
+
+		int* OnHeapNotInTransaction = new int{2};
+		REQUIRE(!AutoRTFM::IsInnerTransactionStack(OnHeapNotInTransaction));
+		delete OnHeapNotInTransaction;
+	}
+
+	AutoRTFM::Commit([&]
+	{
+		int OnStackInTransaction = 3;
+		REQUIRE(AutoRTFM::IsInnerTransactionStack(&OnStackInTransaction));
+
+		int* OnHeapInTransaction = new int{4};
+		REQUIRE(!AutoRTFM::IsInnerTransactionStack(OnHeapInTransaction));
+		delete OnHeapInTransaction;
+
+		AutoRTFM::Commit([&]
+		{
+			// `OnStackInTransaction` is no longer in the innermost scope.
+			REQUIRE(!AutoRTFM::IsInnerTransactionStack(&OnStackInTransaction));
+
+			int OnInnermostStackInTransaction = 5;
+			REQUIRE(AutoRTFM::IsInnerTransactionStack(&OnInnermostStackInTransaction));
+		});
+	});
+}
\ No newline at end of file
