[AutoRTFM] Fix some trivial memory leaks in the AutoRTFMTests

diff --git a/Engine/Source/Programs/AutoRTFMTests/Private/Allocators.cpp b/Engine/Source/Programs/AutoRTFMTests/Private/Allocators.cpp
--- a/Engine/Source/Programs/AutoRTFMTests/Private/Allocators.cpp
+++ b/Engine/Source/Programs/AutoRTFMTests/Private/Allocators.cpp
@@ -128,6 +128,7 @@ TEST_CASE("Allocators.ArrayDelete")
 
 		REQUIRE(AutoRTFM::ETransactionResult::AbortedByRequest == Result);
 		REQUIRE(10 == Data[10]);
+		delete[] Data;
 	}
 
 	SECTION("With Commit")
@@ -198,6 +199,7 @@ TEST_CASE("Allocators.DeleteNoOpts")
 
 		REQUIRE(AutoRTFM::ETransactionResult::AbortedByRequest == Result);
 		REQUIRE(42 == *Data);
+		delete Data;
 	}
 
 	SECTION("With Commit")
diff --git a/Engine/Source/Programs/AutoRTFMTests/Private/SharedPointerTests.cpp b/Engine/Source/Programs/AutoRTFMTests/Private/SharedPointerTests.cpp
--- a/Engine/Source/Programs/AutoRTFMTests/Private/SharedPointerTests.cpp
+++ b/Engine/Source/Programs/AutoRTFMTests/Private/SharedPointerTests.cpp
@@ -155,10 +155,12 @@ TEST_CASE("SharedPointer.NestedTransactionWithPlacementNewlyAllocated")
 		{
 			AutoRTFM::Commit([&]
 				{
-					void* const Memory = MakeMemoryForT<TSharedPtr<int, ESPMode::ThreadSafe>>();
-					TSharedPtr<int, ESPMode::ThreadSafe>* const Foo = new (Memory) TSharedPtr<int, ESPMode::ThreadSafe>(new int(13));
-					TSharedPtr<int, ESPMode::ThreadSafe> Copy = *Foo;
+					using Ptr = TSharedPtr<int, ESPMode::ThreadSafe>;
+					void* const Memory = MakeMemoryForT<Ptr>();
+					Ptr* const Foo = new (Memory) Ptr(new int(13));
+					Ptr Copy = *Foo;
 					Result = *Copy;
+					reinterpret_cast<Ptr*>(Foo)->~Ptr();
 					free(Memory);
 				});
 		});
