[AutoRTFMTests] Fix / silence memory leaks in UObjectTests

Move GTestAnnotation from a global variable, and the FTestAnnotation type to test-local. Having this be global will be treated as a memory leak as the leak detector is run before global destructors.

Use MALLOCLEAK_IGNORE_SCOPE() for calling PerformGarbageCollectionWithIncrementalReachabilityAnalysis(). If the GC does not complete within the time limit, then it'll leak tasks that never complete. Unfortunately, I see no way to flush these.

diff --git a/Engine/Source/Programs/AutoRTFMTests/Private/UObjectTests.cpp b/Engine/Source/Programs/AutoRTFMTests/Private/UObjectTests.cpp
--- a/Engine/Source/Programs/AutoRTFMTests/Private/UObjectTests.cpp
+++ b/Engine/Source/Programs/AutoRTFMTests/Private/UObjectTests.cpp
@@ -1,6 +1,7 @@
 // Copyright Epic Games, Inc. All Rights Reserved.
 
 #include "GenericPlatform/GenericPlatformMisc.h"
+#include "HAL/MallocLeakDetection.h"
 
 THIRD_PARTY_INCLUDES_START
 #include "Catch2Includes.h"
@@ -123,6 +124,9 @@ TEST_CASE("UObject.MarkAsReachable")
 
 	UMyAutoRTFMTestObject* const Object = NewObject<UMyAutoRTFMTestObject>();
 
+	// Somewhat ironically, garbage collection can leak memory.
+	MALLOCLEAK_IGNORE_SCOPE();
+
 	PerformGarbageCollectionWithIncrementalReachabilityAnalysis([Object](int32 index)
 	{
 		if (0 != index)
@@ -142,8 +146,7 @@ TEST_CASE("UObject.MarkAsReachable")
 	SetReachabilityAnalysisTimeLimit(Original);
 }
 
-
-namespace
+TEST_CASE("UObject.TestAddAnnotation")
 {
 	struct FTestAnnotation
 	{
@@ -159,11 +162,9 @@ namespace
 			return TestAnnotationNumber == 32;
 		}
 	};
+
 	FUObjectAnnotationSparse<FTestAnnotation, true> GTestAnnotation;
-}
 
-TEST_CASE("UObject.TestAddAnnotation")
-{
 	SECTION("Create")
 	{
 		UMyAutoRTFMTestObject* Outer = NewObject<UMyAutoRTFMTestObject>();
