[AutoRTFM] Rename AutoRTFM::ForTheRuntime::RecordOpenWrite to AutoRTFM::RecordOpenWrite.

When dealing with external code, non-runtime may need to record writes in the open. For instance, the animation code uses the ISPC compiler, which is not natively RTFM-compatible, but we can run ISPC code in the open and manually track its writes.

diff --git a/Engine/Source/Programs/AutoRTFMTests/Private/OpenAPI.cpp b/Engine/Source/Programs/AutoRTFMTests/Private/OpenAPI.cpp
--- a/Engine/Source/Programs/AutoRTFMTests/Private/OpenAPI.cpp
+++ b/Engine/Source/Programs/AutoRTFMTests/Private/OpenAPI.cpp
@@ -18,14 +18,14 @@ TEST_CASE("OpenAPI.StartAbortAndStartAgain")
 		AutoRTFM::Open([&]()
 		{
 			AutoRTFM::ForTheRuntime::StartTransaction();
-			AutoRTFM::ForTheRuntime::RecordOpenWrite(&valueB);
+			AutoRTFM::RecordOpenWrite(&valueB);
 			valueB=10;
 			AutoRTFM::AbortTransaction();
 
 			AutoRTFM::ForTheRuntime::ClearTransactionStatus();
 
 			AutoRTFM::ForTheRuntime::StartTransaction();
-			AutoRTFM::ForTheRuntime::RecordOpenWrite(&valueC);
+			AutoRTFM::RecordOpenWrite(&valueC);
 			valueC = 30;
 			AutoRTFM::AbortTransaction();
 		});
@@ -55,7 +55,7 @@ TEST_CASE("OpenAPI.RecordDataClosed_Illegal", "[.]")
 	{
 		REQUIRE(AutoRTFM::EContextStatus::OnTrack == AutoRTFM::Close([&]()
 		{
-			AutoRTFM::ForTheRuntime::RecordOpenWrite(&value); // Illegal. Can't record writes explicitly while closed
+			AutoRTFM::RecordOpenWrite(&value); // Illegal. Can't record writes explicitly while closed
 			value = 1;
 		}));
 	});
@@ -71,7 +71,7 @@ TEST_CASE("OpenAPI.WriteDataInTheOpen")
 	{
 		AutoRTFM::Open([&]()
 		{
-			AutoRTFM::ForTheRuntime::RecordOpenWrite(&value);
+			AutoRTFM::RecordOpenWrite(&value);
 			value = 1;
 		});
 	});
@@ -469,7 +469,7 @@ TEST_CASE("OpenAPI.StackWriteCommitInTheOpen1")
 		AutoRTFM::Open([&]()
 		{
 			AutoRTFM::ForTheRuntime::StartTransaction();
-			AutoRTFM::ForTheRuntime::RecordOpenWrite(&value);
+			AutoRTFM::RecordOpenWrite(&value);
 			value = 10;
 			AutoRTFM::ForTheRuntime::CommitTransaction();
 			REQUIRE(value == 10);
@@ -489,7 +489,7 @@ TEST_CASE("OpenAPI.StackWriteCommitInTheOpen2")
 			{
 				AutoRTFM::Open([&]()
 				{
-					AutoRTFM::ForTheRuntime::RecordOpenWrite(&value);
+					AutoRTFM::RecordOpenWrite(&value);
 					value = 10;
 				});
 			}));
@@ -508,7 +508,7 @@ TEST_CASE("OpenAPI.StackWriteAbortInTheOpen1")
 		AutoRTFM::Open([&]()
 		{
 			AutoRTFM::ForTheRuntime::StartTransaction();
-			AutoRTFM::ForTheRuntime::RecordOpenWrite(&value);
+			AutoRTFM::RecordOpenWrite(&value);
 			value = 10;
 			AutoRTFM::AbortTransaction();
 			REQUIRE(value == 0);
@@ -724,7 +724,7 @@ TEST_CASE("OpenAPI.Footgun2")
 				// Unrecorded assignments in the open
 				valueB = 10;
 				valueC = 10;
-				AutoRTFM::ForTheRuntime::RecordOpenWrite(&valueC);
+				AutoRTFM::RecordOpenWrite(&valueC);
 				// valueC was recorded in the open after the change - too late
 			});
 
@@ -761,7 +761,7 @@ TEST_CASE("OpenAPI.StartCloseOnCommit")
 
 	// Setting a value in the open requires us to register the memory address with the transaction
 	value = 42;
-	AutoRTFM::ForTheRuntime::RecordOpenWrite(&value);
+	AutoRTFM::RecordOpenWrite(&value);
 
 	AutoRTFM::ForTheRuntime::CommitTransaction();
 
diff --git a/Engine/Source/Runtime/Core/Public/AutoRTFM/AutoRTFM.h b/Engine/Source/Runtime/Core/Public/AutoRTFM/AutoRTFM.h
--- a/Engine/Source/Runtime/Core/Public/AutoRTFM/AutoRTFM.h
+++ b/Engine/Source/Runtime/Core/Public/AutoRTFM/AutoRTFM.h
@@ -819,6 +819,22 @@ static UE_AUTORTFM_FORCEINLINE void DidFree(void* Ptr)
     autortfm_did_free(Ptr);
 }
 
+// Informs the runtime that a block of memory is about to be overwritten in the open.
+// During a transaction, this allows the runtime to copy the data in preparation for
+// a possible abort. Normally, tracking memory overwrites should be automatically
+// handled by AutoRTFM, but manual overwrite tracking may be required for third-party 
+// libraries or outside compilers (such as ISPC).
+static UE_AUTORTFM_FORCEINLINE void RecordOpenWrite(void* Ptr, size_t Size)
+{
+	autortfm_record_open_write(Ptr, Size);
+}
+
+// Informs the runtime that a block of memory is about to be overwritten.
+template<typename TTYPE> static UE_AUTORTFM_FORCEINLINE void RecordOpenWrite(TTYPE* Ptr)
+{
+	autortfm_record_open_write(Ptr, sizeof(TTYPE));
+}
+
 // A collection of power-user functions that are reserved for use by the AutoRTFM runtime only.
 namespace ForTheRuntime
 {
@@ -937,18 +953,6 @@ namespace ForTheRuntime
 		}
 	};
 
-	// Manually records that the memory span was written in the current transaction.
-	static UE_AUTORTFM_FORCEINLINE void RecordOpenWrite(void* Ptr, size_t Size)
-	{
-		autortfm_record_open_write(Ptr, Size);
-	}
-
-	// Manually records that the memory span was written in the current transaction.
-	template<typename TTYPE> static UE_AUTORTFM_FORCEINLINE void RecordOpenWrite(TTYPE* Ptr)
-	{
-		autortfm_record_open_write(Ptr, sizeof(TTYPE));
-	}
-
 	// Reserved for future.
 	static UE_AUTORTFM_FORCEINLINE void RecordOpenRead(void const*, size_t) {}
 
