[AutoRTFM] Add a test for SOL-5804

This has been fixed some time ago, but adding a test to prove it.

diff --git a/Engine/Source/Programs/AutoRTFMTests/Private/Abort.cpp b/Engine/Source/Programs/AutoRTFMTests/Private/Abort.cpp
--- a/Engine/Source/Programs/AutoRTFMTests/Private/Abort.cpp
+++ b/Engine/Source/Programs/AutoRTFMTests/Private/Abort.cpp
@@ -623,3 +623,29 @@ TEST_CASE("Abort.LanguageThroughOpen")
 	REQUIRE(AutoRTFM::ETransactionResult::AbortedByLanguage == Result);
 	REQUIRE(false == bTouched);
 }
+
+// Test for SOL-5804
+TEST_CASE("Abort.StackWriteToOuterOpen")
+{
+	AutoRTFM::EContextStatus CloseStatus = AutoRTFM::EContextStatus::Idle;
+
+	const AutoRTFM::ETransactionResult Result = AutoRTFM::Transact([&]
+	{
+		AutoRTFM::Open([&]
+		{
+			int Value = 10;
+
+			CloseStatus = AutoRTFM::Close([&]
+			{
+				// On stack inside transaction.
+				// Value not be reverted as part of abort.
+				Value = 20;
+			});
+		});
+
+		AutoRTFM::AbortTransaction();
+	});
+
+	REQUIRE(AutoRTFM::ETransactionResult::AbortedByRequest == Result);
+	REQUIRE(AutoRTFM::EContextStatus::OnTrack == CloseStatus);
+}
