[AutoRTFM] Improve the Abort.StackWriteToOuterOpen test

Try and detect bad stack stomps.

diff --git a/Engine/Source/Programs/AutoRTFMTests/Private/Abort.cpp b/Engine/Source/Programs/AutoRTFMTests/Private/Abort.cpp
--- a/Engine/Source/Programs/AutoRTFMTests/Private/Abort.cpp
+++ b/Engine/Source/Programs/AutoRTFMTests/Private/Abort.cpp
@@ -627,25 +627,57 @@ TEST_CASE("Abort.LanguageThroughOpen")
 // Test for SOL-5804
 TEST_CASE("Abort.StackWriteToOuterOpen")
 {
+	std::string_view TestResult;
 	AutoRTFM::EContextStatus CloseStatus = AutoRTFM::EContextStatus::Idle;
+	bool WritesUndone = true;
 
-	const AutoRTFM::ETransactionResult Result = AutoRTFM::Transact([&]
+	const AutoRTFM::ETransactionResult TransactionResult = AutoRTFM::Transact([&]
 	{
 		AutoRTFM::Open([&]
 		{
-			int Value = 10;
+			std::array<int, 64> Values{};
 
 			CloseStatus = AutoRTFM::Close([&]
 			{
+				// On stack outside transaction.
+				// Should be reverted as part of the abort.
+				WritesUndone = false;
+
 				// On stack inside transaction.
-				// Value not be reverted as part of abort.
-				Value = 20;
+				// Writes should not be reverted as part of the abort.
+				for (size_t I = 0; I < Values.size(); I++)
+				{
+					Values[I] = static_cast<int>(I * 10);
+				}
 			});
 		});
 
+		// If any of the variables on the stack within the Open() get written to
+		// on abort, then it should change the values of this array.
+		std::array<int, 64> StackGuard{};
+
+		// The OnAbort handler should be called *after* the memory is reverted.
+		AutoRTFM::OnAbort([&]
+		{
+			if (!WritesUndone)
+			{
+				TestResult = "OnAbort was called without first reverting memory";
+			}
+			else if (StackGuard != std::array<int, 64>{})
+			{
+				TestResult = "StackGuard was corrupted";
+			}
+			else
+			{
+				TestResult = "Success";
+			}
+		});
+
+		// Do the abort!
 		AutoRTFM::AbortTransaction();
 	});
 
-	REQUIRE(AutoRTFM::ETransactionResult::AbortedByRequest == Result);
+	REQUIRE(AutoRTFM::ETransactionResult::AbortedByRequest == TransactionResult);
 	REQUIRE(AutoRTFM::EContextStatus::OnTrack == CloseStatus);
+	REQUIRE("Success" == TestResult);
 }
