[AutoRTFMTests] Check for memory leaks, fail if any are found.

diff --git a/Engine/Source/Programs/AutoRTFMTests/AutoRTFMTests.Target.cs b/Engine/Source/Programs/AutoRTFMTests/AutoRTFMTests.Target.cs
--- a/Engine/Source/Programs/AutoRTFMTests/AutoRTFMTests.Target.cs
+++ b/Engine/Source/Programs/AutoRTFMTests/AutoRTFMTests.Target.cs
@@ -56,5 +56,8 @@ public class AutoRTFMTestsTarget : TargetRules
 		{
 			GlobalDefinitions.Add("ALLOW_LOG_FILE=1");
 		}
+
+		GlobalDefinitions.Add("MALLOC_LEAKDETECTION=1");
+		GlobalDefinitions.Add("PLATFORM_USES_FIXED_GMalloc_CLASS=0");
 	}
 }
diff --git a/Engine/Source/Programs/AutoRTFMTests/Private/AutoRTFMTests.cpp b/Engine/Source/Programs/AutoRTFMTests/Private/AutoRTFMTests.cpp
--- a/Engine/Source/Programs/AutoRTFMTests/Private/AutoRTFMTests.cpp
+++ b/Engine/Source/Programs/AutoRTFMTests/Private/AutoRTFMTests.cpp
@@ -1,14 +1,52 @@
 // Copyright Epic Games, Inc. All Rights Reserved.
 
 #include "CoreMinimal.h"
+#include "Misc/EventPool.h"
+#include "Misc/LazySingleton.h"
 #include "RequiredProgramMainCPPInclude.h" // required for ue programs
+#include "HAL/MallocLeakDetection.h"
 
 IMPLEMENT_APPLICATION(AutoRTFMTests, "AutoRTFMTests");
 
 #define CATCH_AMALGAMATED_CUSTOM_MAIN
 #include "catch_amalgamated.cpp"
 
-int main(int ArgC, const char* ArgV[]) {
+class FListener final : public Catch::EventListenerBase
+{
+public:
+	using Catch::EventListenerBase::EventListenerBase;
+
+	void testCaseStarting(const Catch::TestCaseInfo&) override
+	{
+		FMallocLeakDetection::Get().SetAllocationCollection(true);
+	}
+
+	void testCaseEnded(const Catch::TestCaseStats&) override
+	{
+		FMallocLeakDetection::Get().SetAllocationCollection(false);
+	}
+
+	void testRunEnded(const Catch::TestRunStats&) override
+	{
+	}
+};
+
+CATCH_REGISTER_LISTENER(FListener)
+
+// Checks for memory leaks. Returns true if no leaks were found, otherwise all
+// leaks are printed to stderr and false is returned.
+bool CheckNoMemoryLeaks();
+
+int main(int ArgC, const char* ArgV[])
+{
+	{
+		// measure_environment lazily allocates an Environment pointer on first
+		// call, and holds this in a static variable. Pre-allocate this with the
+		// leak detector disabled to prevent this being reported as a leak.
+		MALLOCLEAK_IGNORE_SCOPE();
+		Catch::Benchmark::Detail::measure_environment<Catch::Benchmark::default_clock>();
+	}
+
 	Catch::Session Session;
 
 	bool NoRetry = false;
@@ -51,13 +89,46 @@ int main(int ArgC, const char* ArgV[]) {
 	// We don't want to trigger ensure's on abort because we are going to test that.
 	AutoRTFM::ForTheRuntime::SetEnsureOnAbortByLanguage(false);
 
-	const int Result = Session.run();
+	int Result = Session.run();
 
 	FPlatformMisc::RequestExit(false);
 
+	CollectGarbage(GARBAGE_COLLECTION_KEEPFLAGS, /* bPerformFullPurge */ true);
+
 	FEngineLoop::AppPreExit();
 	FModuleManager::Get().UnloadModulesAtShutdown();
 	FEngineLoop::AppExit();
+	TLazySingleton<TEventPool<EEventMode::AutoReset>>::Get().EmptyPool();
+	TLazySingleton<TEventPool<EEventMode::ManualReset>>::Get().EmptyPool();
+
+	if (0 == Result && !CheckNoMemoryLeaks())
+	{
+		Result = -1;
+	}
 
 	return Result;
 }
+
+bool CheckNoMemoryLeaks()
+{
+	class FOutputDeviceStderr final : public FOutputDevice
+	{
+	public:
+		void Serialize(const TCHAR* V, ELogVerbosity::Type Verbosity, const class FName& Category) override
+		{
+			std::cerr << reinterpret_cast<const char*>(StringCast<UTF8CHAR>(V).Get()) << std::endl;
+		}
+	};
+
+	FOutputDeviceStderr OutputDevice;
+	FMallocLeakReportOptions Options;
+	Options.OutputDevice = &OutputDevice;
+	int32 NumLeaks = FMallocLeakDetection::Get().DumpOpenCallstacks(TEXT("AutoRTFMTests"), Options);
+	if (NumLeaks > 0)
+	{
+		std::cerr << NumLeaks << " memory leaks detected" << std::endl;
+		return false;
+	}
+
+	return true;
+}
