[AutoRTFMTests] Destruct the Catch2 session before checking for memory leaks

On linux, Catch2 holds on to a string buffer which is flagged by the leak detector.

diff --git a/Engine/Source/Programs/AutoRTFMTests/Private/AutoRTFMTests.cpp b/Engine/Source/Programs/AutoRTFMTests/Private/AutoRTFMTests.cpp
--- a/Engine/Source/Programs/AutoRTFMTests/Private/AutoRTFMTests.cpp
+++ b/Engine/Source/Programs/AutoRTFMTests/Private/AutoRTFMTests.cpp
@@ -37,7 +37,7 @@ CATCH_REGISTER_LISTENER(FListener)
 // leaks are printed to stderr and false is returned.
 bool CheckNoMemoryLeaks();
 
-int main(int ArgC, const char* ArgV[])
+int RunTests(int ArgC, const char* ArgV[])
 {
 	{
 		// measure_environment lazily allocates an Environment pointer on first
@@ -101,6 +101,12 @@ int main(int ArgC, const char* ArgV[])
 	TLazySingleton<TEventPool<EEventMode::AutoReset>>::Get().EmptyPool();
 	TLazySingleton<TEventPool<EEventMode::ManualReset>>::Get().EmptyPool();
 
+	return Result;
+}
+
+int main(int ArgC, const char* ArgV[])
+{
+	int Result = RunTests(ArgC, ArgV);
 	if (0 == Result && !CheckNoMemoryLeaks())
 	{
 		Result = -1;
