[AutoRTFM] Add unit tests for non-cascaded aborts.

Most of our tests focused on cascaded aborts; this verifies that non-cascaded aborts work when transactions are nested, and that OnAbort handlers are triggered at appropriate times.

diff --git a/Engine/Source/Programs/AutoRTFMTests/Private/Abort.cpp b/Engine/Source/Programs/AutoRTFMTests/Private/Abort.cpp
--- a/Engine/Source/Programs/AutoRTFMTests/Private/Abort.cpp
+++ b/Engine/Source/Programs/AutoRTFMTests/Private/Abort.cpp
@@ -131,6 +131,52 @@ TEST_CASE("Abort.TransactionInOnAbort")
 	REQUIRE(AutoRTFM::ETransactionResult::AbortedByRequest == Result);
 }
 
+TEST_CASE("Abort.AbortInInnerTransaction")
+{
+	int Value = 1;
+
+	AutoRTFM::Transact([&]
+	{
+		Value = 2;
+		AutoRTFM::Transact([&]
+		{
+			Value = 3;
+			AutoRTFM::AbortTransaction();  // Only cancels innermost nest.
+		});
+	});
+
+	REQUIRE(Value == 2);
+}
+
+TEST_CASE("Abort.OnAbortInOuterTransaction_AbortInInnerTransaction")
+{
+	int Value = 1;
+
+	AutoRTFM::Transact([&]
+	{
+		AutoRTFM::OnAbort([&]
+		{
+			if (AutoRTFM::ForTheRuntime::GetRetryTransaction() == AutoRTFM::ForTheRuntime::NoRetry)
+			{
+				REQUIRE(!"This should not be reached.");
+				Value = 2;
+			}
+		});
+
+		AutoRTFM::Transact([&]
+		{
+			REQUIRE(Value == 1);
+			Value = 3;
+			AutoRTFM::AbortTransaction();  // Only cancels innermost nest.
+		});
+
+		REQUIRE(Value == 1);
+		Value = 4;
+	});
+
+	REQUIRE(Value == 4);
+}
+
 TEST_CASE("Abort.Cascade")
 {
 	bool bTouched = false;
