[AutoRTFM] Add more cstring stdlib support

diff --git a/Engine/Source/Programs/AutoRTFMTests/Private/CString.cpp b/Engine/Source/Programs/AutoRTFMTests/Private/CString.cpp
--- a/Engine/Source/Programs/AutoRTFMTests/Private/CString.cpp
+++ b/Engine/Source/Programs/AutoRTFMTests/Private/CString.cpp
@@ -43,7 +43,7 @@ static constexpr auto* kPercentNWarning = TEXT("AutoRTFM does not support format
 
 // CheckWCharFunction calls Function with WChar outside a transaction, in an
 // aborted transaction and committed transaction. The value returned when called
-// in the committed transation is expected to match the value returned when
+// in the committed transaction is expected to match the value returned when
 // called outside the transaction.
 // Function must be the deterministic function with the signature 'int(wchar_t)'.
 template<typename FUNCTION>
@@ -79,6 +79,79 @@ void CheckWCharFunction(FUNCTION&& Function, wchar_t WChar)
 	}
 }
 
+// CheckStringToNumberFunction calls Function with String outside a transaction,
+// in an aborted transaction and committed transaction, once with a valid
+// end pointer and again with a null end pointer. The value returned and the
+// assigned end-pointer when called in the committed transaction is expected to
+// match the values when called outside the transaction.
+// Function must be the deterministic function with the signature
+// 'RESULT(const CHAR* String, CHAR** EndPtr, ADDITIONAL_ARGS...)'.
+template<typename RESULT, typename CHAR, typename ... ADDITIONAL_ARGS>
+void CheckStringToNumberFunction(RESULT Function(const CHAR*, CHAR**, ADDITIONAL_ARGS...), const CHAR* String, ADDITIONAL_ARGS... AdditionalArgs)
+{
+	CHAR* ExpectEndPtr = nullptr;
+	RESULT ExpectValue = Function(String, &ExpectEndPtr, AdditionalArgs...);
+
+	CHAR* GotEndPtr = nullptr;
+	RESULT GotValue = 0;
+
+	SECTION("With EndPtr")
+	{
+		SECTION("With Abort")
+		{
+			CLANG_BEGIN_DISABLE_OPTIMIZATIONS
+			AutoRTFM::ETransactionResult Result = AutoRTFM::Transact([&]
+			{
+				GotValue = Function(String, &GotEndPtr, AdditionalArgs...);
+				AutoRTFM::AbortTransaction();
+			});
+			CLANG_END_DISABLE_OPTIMIZATIONS
+
+			REQUIRE(AutoRTFM::ETransactionResult::AbortedByRequest == Result);
+			REQUIRE(0 == GotValue);
+			REQUIRE(nullptr == GotEndPtr);
+		}
+
+		SECTION("With Commit")
+		{
+			CLANG_BEGIN_DISABLE_OPTIMIZATIONS
+			AutoRTFM::Commit([&]
+			{
+				GotValue = Function(String, &GotEndPtr, AdditionalArgs...);
+			});
+			CLANG_END_DISABLE_OPTIMIZATIONS
+
+			REQUIRE(ExpectValue == GotValue);
+			REQUIRE(ExpectEndPtr == GotEndPtr);
+		}
+	}
+	SECTION("Null EndPtr")
+	{
+		SECTION("With Abort")
+		{
+			CLANG_BEGIN_DISABLE_OPTIMIZATIONS
+			AutoRTFM::ETransactionResult Result = AutoRTFM::Transact([&]
+			{
+				GotValue = Function(String, nullptr, AdditionalArgs...);
+				AutoRTFM::AbortTransaction();
+			});
+			CLANG_END_DISABLE_OPTIMIZATIONS
+
+			REQUIRE(AutoRTFM::ETransactionResult::AbortedByRequest == Result);
+			REQUIRE(0 == GotValue);
+		}
+
+		SECTION("With Commit")
+		{
+			CLANG_BEGIN_DISABLE_OPTIMIZATIONS
+			AutoRTFM::Commit([&] { GotValue = Function(String, nullptr, AdditionalArgs...); });
+			CLANG_END_DISABLE_OPTIMIZATIONS
+
+			REQUIRE(ExpectValue == GotValue);
+		}
+	}
+}
+
 // A helper class that for the lifetime of the object, intercepts and records
 // UE_LOG warnings.
 class FCaptureWarningContext : private FFeedbackContext
@@ -665,76 +738,6 @@ TEST_CASE("CString.strlen")
 	}
 }
 
-TEST_CASE("CString.strtol")
-{
-	const char* Str = "123xyz";
-
-	SECTION("With end pointer")
-	{
-		SECTION("With Abort")
-		{
-			long int Value = 0;
-			char* EndPtr = nullptr;
-
-			CLANG_BEGIN_DISABLE_OPTIMIZATIONS
-			AutoRTFM::ETransactionResult Result = AutoRTFM::Transact([&]
-			{
-				Value = strtol(Str, &EndPtr, 10);
-				AutoRTFM::AbortTransaction();
-			});
-			CLANG_END_DISABLE_OPTIMIZATIONS
-
-			REQUIRE(AutoRTFM::ETransactionResult::AbortedByRequest == Result);
-			REQUIRE(0 == Value);
-			REQUIRE(nullptr == EndPtr);
-		}
-
-		SECTION("With Commit")
-		{
-			long int Value = 0;
-			char* EndPtr = nullptr;
-
-			CLANG_BEGIN_DISABLE_OPTIMIZATIONS
-			AutoRTFM::Commit([&] { Value = strtol(Str, &EndPtr, 10); });
-			CLANG_END_DISABLE_OPTIMIZATIONS
-
-			REQUIRE(123 == Value);
-			REQUIRE((Str + 3) == EndPtr);
-		}
-	}
-
-	SECTION("Without end pointer")
-	{
-		SECTION("With Abort")
-		{
-			long int Value = 0;
-
-			CLANG_BEGIN_DISABLE_OPTIMIZATIONS
-			AutoRTFM::ETransactionResult Result = AutoRTFM::Transact([&]
-			{
-				Value = strtol(Str, nullptr, 10);
-				AutoRTFM::AbortTransaction();
-			});
-			CLANG_END_DISABLE_OPTIMIZATIONS
-
-			REQUIRE(AutoRTFM::ETransactionResult::AbortedByRequest == Result);
-			REQUIRE(0 == Value);
-		}
-
-		SECTION("With Commit")
-		{
-			long int Value = 0;
-
-			CLANG_BEGIN_DISABLE_OPTIMIZATIONS
-			AutoRTFM::Commit([&] { Value = strtol(Str, nullptr, 10); });
-			CLANG_END_DISABLE_OPTIMIZATIONS
-
-			REQUIRE(123 == Value);
-		}
-	}
-}
-
-
 TEST_CASE("CString.wcscpy")
 {
 	MSVC_BEGIN_DISABLE_WARN_UNSAFE_FUNCTION
@@ -1032,6 +1035,103 @@ TEST_CASE("CString.wcslen")
 	}
 }
 
+TEST_CASE("CString.wcstod")
+{
+	CheckStringToNumberFunction(wcstod, L"42.5abc");
+}
+
+TEST_CASE("CString.wcstof")
+{
+	CheckStringToNumberFunction(wcstof, L"42.5abc");
+}
+
+TEST_CASE("CString.wcstold")
+{
+	CheckStringToNumberFunction(wcstold, L"42.5abc");
+}
+
+TEST_CASE("CString.strtol")
+{
+	CheckStringToNumberFunction(strtol, "42.5abc", 10);
+}
+
+TEST_CASE("CString.wcstol")
+{
+	CheckStringToNumberFunction(wcstol, L"42abc", 10);
+}
+
+TEST_CASE("CString.wcstoll")
+{
+	CheckStringToNumberFunction(wcstoll, L"42abc", 10);
+}
+
+TEST_CASE("CString.wcstoul")
+{
+	CheckStringToNumberFunction(wcstoul, L"42abc", 10);
+}
+
+TEST_CASE("CString.wcstoull")
+{
+	CheckStringToNumberFunction(wcstoull, L"42abc", 10);
+}
+
+#if PLATFORM_WINDOWS
+TEST_CASE("CString._wcstoi64")
+{
+	CheckStringToNumberFunction(_wcstoi64, L"42abc", 10);
+}
+
+TEST_CASE("CString._wcstoui64")
+{
+	CheckStringToNumberFunction(_wcstoui64, L"42abc", 10);
+}
+
+TEST_CASE("CString._wcstod_l")
+{
+	CheckStringToNumberFunction(_wcstod_l, L"42.5abc", _get_current_locale());
+}
+
+TEST_CASE("CString._wcstof_l")
+{
+	CheckStringToNumberFunction(_wcstof_l, L"42.5abc", _get_current_locale());
+}
+
+TEST_CASE("CString._wcstold_l")
+{
+	CheckStringToNumberFunction(_wcstold_l, L"42.5abc", _get_current_locale());
+}
+
+TEST_CASE("CString._wcstol_l")
+{
+	CheckStringToNumberFunction(_wcstol_l, L"42abc", 10, _get_current_locale());
+}
+
+TEST_CASE("CString._wcstoll_l")
+{
+	CheckStringToNumberFunction(_wcstoll_l, L"42abc", 10, _get_current_locale());
+}
+
+TEST_CASE("CString._wcstoul_l")
+{
+	CheckStringToNumberFunction(_wcstoul_l, L"42abc", 10, _get_current_locale());
+}
+
+TEST_CASE("CString._wcstoull_l")
+{
+	CheckStringToNumberFunction(_wcstoull_l, L"42abc", 10, _get_current_locale());
+}
+
+TEST_CASE("CString._wcstoi64_l")
+{
+	CheckStringToNumberFunction(_wcstoi64_l, L"42abc", 10, _get_current_locale());
+}
+
+TEST_CASE("CString._wcstoui64_l")
+{
+	CheckStringToNumberFunction(_wcstoui64_l, L"42abc", 10, _get_current_locale());
+}
+
+#endif
 
 TEST_CASE("CString.iswupper")
 {
diff --git a/Engine/Source/Runtime/Core/Private/AutoRTFM/StdLibSupport.cpp b/Engine/Source/Runtime/Core/Private/AutoRTFM/StdLibSupport.cpp
--- a/Engine/Source/Runtime/Core/Private/AutoRTFM/StdLibSupport.cpp
+++ b/Engine/Source/Runtime/Core/Private/AutoRTFM/StdLibSupport.cpp
@@ -216,24 +216,91 @@ UE_AUTORTFM_REGISTER_SELF_OVERLOADED_FUNCTION(const char* (const char*, int), st
 UE_AUTORTFM_REGISTER_SELF_OVERLOADED_FUNCTION(const char* (const char*, const char*), strstr);
 UE_AUTORTFM_REGISTER_SELF_FUNCTION(strlen);
 
-long RTFM_strtol(char const* String, char** EndPtr, int Radix) throw()
+UE_AUTORTFM_REGISTER_SELF_OVERLOADED_FUNCTION(const wchar_t* (const wchar_t*, wchar_t), wcschr);
+UE_AUTORTFM_REGISTER_SELF_OVERLOADED_FUNCTION(wchar_t* (wchar_t*, wchar_t), wcschr);
+UE_AUTORTFM_REGISTER_SELF_OVERLOADED_FUNCTION(wchar_t* (wchar_t*, const wchar_t*), wcsstr);
+UE_AUTORTFM_REGISTER_SELF_OVERLOADED_FUNCTION(const wchar_t*(const wchar_t*, const wchar_t*), wcsstr);
+UE_AUTORTFM_REGISTER_SELF_FUNCTION(wcscmp);
+UE_AUTORTFM_REGISTER_SELF_FUNCTION(wcslen);
+
+template<typename RETURN, typename CHAR, RETURN FN(const CHAR*, CHAR**)>
+RETURN RTFM_StringToFloat(const CHAR* String, CHAR** EndPtr) throw()
 {
     if (nullptr != EndPtr)
     {
         FContext* Context = FContext::Get();
-        Context->RecordWrite(EndPtr, sizeof(char*));
+        Context->RecordWrite(EndPtr, sizeof(CHAR*));
     }
+    return FN(String, EndPtr);
+}
 
-    return strtol(String, EndPtr, Radix);
+template<typename RETURN, typename CHAR, RETURN FN(const CHAR*, CHAR**, int)>
+RETURN RTFM_StringToInt(const CHAR* String, CHAR** EndPtr, int Radix) throw()
+{
+    if (nullptr != EndPtr)
+    {
+        FContext* Context = FContext::Get();
+        Context->RecordWrite(EndPtr, sizeof(CHAR*));
+    }
+    return FN(String, EndPtr, Radix);
 }
-UE_AUTORTFM_REGISTER_OPEN_FUNCTION(strtol);
 
-UE_AUTORTFM_REGISTER_SELF_OVERLOADED_FUNCTION(const wchar_t* (const wchar_t*, wchar_t), wcschr);
-UE_AUTORTFM_REGISTER_SELF_OVERLOADED_FUNCTION(wchar_t* (wchar_t*, wchar_t), wcschr);
-UE_AUTORTFM_REGISTER_SELF_OVERLOADED_FUNCTION(wchar_t* (wchar_t*, const wchar_t*), wcsstr);
-UE_AUTORTFM_REGISTER_SELF_OVERLOADED_FUNCTION(const wchar_t*(const wchar_t*, const wchar_t*), wcsstr);
-UE_AUTORTFM_REGISTER_SELF_FUNCTION(wcscmp);
-UE_AUTORTFM_REGISTER_SELF_FUNCTION(wcslen);
+UE_AUTORTFM_REGISTER_OPEN_FUNCTION_EXPLICIT(strtol, (RTFM_StringToInt<long, char, strtol>));
+UE_AUTORTFM_REGISTER_OPEN_FUNCTION_EXPLICIT(wcstod, (RTFM_StringToFloat<double, wchar_t, wcstod>));
+UE_AUTORTFM_REGISTER_OPEN_FUNCTION_EXPLICIT(wcstof, (RTFM_StringToFloat<float, wchar_t, wcstof>));
+UE_AUTORTFM_REGISTER_OPEN_FUNCTION_EXPLICIT(wcstold, (RTFM_StringToFloat<long double, wchar_t, wcstold>));
+UE_AUTORTFM_REGISTER_OPEN_FUNCTION_EXPLICIT(wcstol, (RTFM_StringToInt<long, wchar_t, wcstol>));
+UE_AUTORTFM_REGISTER_OPEN_FUNCTION_EXPLICIT(wcstoll, (RTFM_StringToInt<long long, wchar_t, wcstoll>));
+UE_AUTORTFM_REGISTER_OPEN_FUNCTION_EXPLICIT(wcstoul, (RTFM_StringToInt<unsigned long, wchar_t, wcstoul>));
+UE_AUTORTFM_REGISTER_OPEN_FUNCTION_EXPLICIT(wcstoull, (RTFM_StringToInt<unsigned long long, wchar_t, wcstoull>));
+
+#if PLATFORM_WINDOWS
+template<typename RETURN, typename CHAR, RETURN FN(const CHAR*, CHAR**, _locale_t)>
+RETURN RTFM_StringToFloat(const CHAR* String, CHAR** EndPtr, _locale_t Locale)
+{
+    if (nullptr != EndPtr)
+    {
+        FContext* Context = FContext::Get();
+        Context->RecordWrite(EndPtr, sizeof(CHAR*));
+    }
+    return FN(String, EndPtr, Locale);
+}
+
+template<typename RETURN, typename CHAR, RETURN FN(const CHAR*, CHAR**, int, _locale_t)>
+RETURN RTFM_StringToInt(const CHAR* String, CHAR** EndPtr, int Radix, _locale_t Locale)
+{
+    if (nullptr != EndPtr)
+    {
+        FContext* Context = FContext::Get();
+        Context->RecordWrite(EndPtr, sizeof(CHAR*));
+    }
+    return FN(String, EndPtr, Radix, Locale);
+}
+
+UE_AUTORTFM_REGISTER_OPEN_FUNCTION_EXPLICIT(_wcstoi64, (RTFM_StringToInt<long long, wchar_t, _wcstoi64>));
+UE_AUTORTFM_REGISTER_OPEN_FUNCTION_EXPLICIT(_wcstoui64, (RTFM_StringToInt<unsigned long long, wchar_t, _wcstoui64>));
+UE_AUTORTFM_REGISTER_OPEN_FUNCTION_EXPLICIT(_wcstod_l, (RTFM_StringToFloat<double, wchar_t, _wcstod_l>));
+UE_AUTORTFM_REGISTER_OPEN_FUNCTION_EXPLICIT(_wcstof_l, (RTFM_StringToFloat<float, wchar_t, _wcstof_l>));
+UE_AUTORTFM_REGISTER_OPEN_FUNCTION_EXPLICIT(_wcstold_l, (RTFM_StringToFloat<long double, wchar_t, _wcstold_l>));
+UE_AUTORTFM_REGISTER_OPEN_FUNCTION_EXPLICIT(_wcstol_l, (RTFM_StringToInt<long, wchar_t, _wcstol_l>));
+UE_AUTORTFM_REGISTER_OPEN_FUNCTION_EXPLICIT(_wcstoll_l, (RTFM_StringToInt<long long, wchar_t, _wcstoll_l>));
+UE_AUTORTFM_REGISTER_OPEN_FUNCTION_EXPLICIT(_wcstoul_l, (RTFM_StringToInt<unsigned long, wchar_t, _wcstoul_l>));
+UE_AUTORTFM_REGISTER_OPEN_FUNCTION_EXPLICIT(_wcstoull_l, (RTFM_StringToInt<unsigned long long, wchar_t, _wcstoull_l>));
+UE_AUTORTFM_REGISTER_OPEN_FUNCTION_EXPLICIT(_wcstoi64_l, (RTFM_StringToInt<long long, wchar_t, _wcstoi64_l>));
+UE_AUTORTFM_REGISTER_OPEN_FUNCTION_EXPLICIT(_wcstoui64_l, (RTFM_StringToInt<unsigned long long, wchar_t, _wcstoui64_l>));
+
+UE_AUTORTFM_REGISTER_SELF_FUNCTION(_wtof);
+UE_AUTORTFM_REGISTER_SELF_FUNCTION(_wtof);
+UE_AUTORTFM_REGISTER_SELF_FUNCTION(_wtol);
+UE_AUTORTFM_REGISTER_SELF_FUNCTION(_wtoll);
+UE_AUTORTFM_REGISTER_SELF_FUNCTION(_wtoi);
+UE_AUTORTFM_REGISTER_SELF_FUNCTION(_wtoi64);
+UE_AUTORTFM_REGISTER_SELF_FUNCTION(_wtof_l);
+UE_AUTORTFM_REGISTER_SELF_FUNCTION(_wtol_l);
+UE_AUTORTFM_REGISTER_SELF_FUNCTION(_wtoll_l);
+UE_AUTORTFM_REGISTER_SELF_FUNCTION(_wtoi_l);
+UE_AUTORTFM_REGISTER_SELF_FUNCTION(_wtoi64_l);
+#endif
 
 UE_AUTORTFM_REGISTER_SELF_FUNCTION(iswupper);
 UE_AUTORTFM_REGISTER_SELF_FUNCTION(iswlower);
