[AutoRTFM] Handle mixing TMulticastDelegateBase calls from both open and closed transactions

diff --git a/Engine/Source/Programs/AutoRTFMTests/Private/Delegate.cpp b/Engine/Source/Programs/AutoRTFMTests/Private/Delegate.cpp
new file mode 100644
--- /dev/null
+++ b/Engine/Source/Programs/AutoRTFMTests/Private/Delegate.cpp
@@ -0,0 +1,133 @@
+// Copyright Epic Games, Inc. All Rights Reserved.
+
+#include "Catch2Includes.h"
+#include <AutoRTFM/AutoRTFM.h>
+#include "Delegates/Delegate.h"
+
+// Temporarily changes the AutoRTFM retry mode for the lifetime of the FScopedRetry object.
+struct FScopedRetry
+{
+	FScopedRetry(AutoRTFM::ForTheRuntime::EAutoRTFMRetryTransactionState NewRetry)
+		: OldRetry(AutoRTFM::ForTheRuntime::GetRetryTransaction())
+	{
+		AutoRTFM::ForTheRuntime::SetRetryTransaction(NewRetry);
+	}
+
+	~FScopedRetry()
+	{
+		AutoRTFM::ForTheRuntime::SetRetryTransaction(OldRetry);
+	}
+
+	AutoRTFM::ForTheRuntime::EAutoRTFMRetryTransactionState OldRetry;
+};
+
+TEST_CASE("Delegate.Broadcast")
+{
+	// Tests are sensitive to retries. Disable for these tests.
+	FScopedRetry Retry(AutoRTFM::ForTheRuntime::EAutoRTFMRetryTransactionState::NoRetry);
+
+	TMulticastDelegate<void()> Delegate;
+
+	int CallCount = 0;
+	Delegate.AddLambda([&CallCount]
+	{ 
+		AutoRTFM::Open([&CallCount] { CallCount++; });
+	});
+
+	SECTION("Transact(Broadcast)")
+	{
+		AutoRTFM::ETransactionResult Result = AutoRTFM::Transact([&]
+		{
+			Delegate.Broadcast();
+		});
+
+		REQUIRE(AutoRTFM::ETransactionResult::Committed == Result);
+		REQUIRE(1 == CallCount);
+	}
+
+	SECTION("Transact(Open(Broadcast), Broadcast)")
+	{
+		AutoRTFM::ETransactionResult Result = AutoRTFM::Transact([&]
+		{
+			AutoRTFM::Open([&]
+			{
+				Delegate.Broadcast();
+			});
+			Delegate.Broadcast();
+		});
+
+		REQUIRE(AutoRTFM::ETransactionResult::Committed == Result);
+		REQUIRE(2 == CallCount);
+	}
+
+	SECTION("Transact(Broadcast, Open(Broadcast))")
+	{
+		AutoRTFM::ETransactionResult Result = AutoRTFM::Transact([&]
+		{
+			Delegate.Broadcast();
+			AutoRTFM::Open([&]
+			{
+				Delegate.Broadcast();
+			});
+		});
+
+		REQUIRE(AutoRTFM::ETransactionResult::Committed == Result);
+		REQUIRE(2 == CallCount);
+	}
+
+	SECTION("Transact(Broadcast, Abort)")
+	{
+		AutoRTFM::ETransactionResult Result = AutoRTFM::Transact([&]
+		{
+			Delegate.Broadcast();
+			AutoRTFM::AbortTransaction();
+		});
+
+		REQUIRE(AutoRTFM::ETransactionResult::AbortedByRequest == Result);
+		REQUIRE(1 == CallCount);
+	}
+
+	SECTION("Transact(Open(Broadcast), Broadcast, Abort)")
+	{
+		AutoRTFM::ETransactionResult Result = AutoRTFM::Transact([&]
+		{
+			AutoRTFM::Open([&]
+			{
+				Delegate.Broadcast();
+			});
+			Delegate.Broadcast();
+			AutoRTFM::AbortTransaction();
+		});
+
+		REQUIRE(AutoRTFM::ETransactionResult::AbortedByRequest == Result);
+		REQUIRE(2 == CallCount);
+	}
+
+	SECTION("Transact(Broadcast, Open(Broadcast), Abort)")
+	{
+		AutoRTFM::ETransactionResult Result = AutoRTFM::Transact([&]
+		{
+			Delegate.Broadcast();
+			AutoRTFM::Open([&]
+			{
+				Delegate.Broadcast();
+			});
+			AutoRTFM::AbortTransaction();
+		});
+
+		REQUIRE(AutoRTFM::ETransactionResult::AbortedByRequest == Result);
+		REQUIRE(2 == CallCount);
+	}
+
+	SECTION("AddStatic(Abort) Transact(Broadcast)")
+	{
+		Delegate.AddStatic([] { AutoRTFM::AbortTransaction(); });
+		AutoRTFM::ETransactionResult Result = AutoRTFM::Transact([&]
+		{
+			Delegate.Broadcast();
+		});
+
+		REQUIRE(AutoRTFM::ETransactionResult::AbortedByRequest == Result);
+		REQUIRE(0 == CallCount);
+	}
+}
diff --git a/Engine/Source/Runtime/Core/Public/Delegates/MulticastDelegateBase.h b/Engine/Source/Runtime/Core/Public/Delegates/MulticastDelegateBase.h
--- a/Engine/Source/Runtime/Core/Public/Delegates/MulticastDelegateBase.h
+++ b/Engine/Source/Runtime/Core/Public/Delegates/MulticastDelegateBase.h
@@ -2,6 +2,7 @@
 
 #pragma once
 
+#include "AutoRTFM/AutoRTFM.h"
 #include "CoreTypes.h"
 #include "Containers/ContainerAllocationPolicies.h"
 #include "Containers/Array.h"
@@ -381,13 +382,15 @@ private:
 	/** Increments the lock counter for the invocation list. */
 	inline void LockInvocationList( ) const
 	{
-		++InvocationListLockCount;
+		UE_AUTORTFM_OPEN { ++InvocationListLockCount; };
+		UE_AUTORTFM_ONABORT(this) { --InvocationListLockCount; };
 	}
 
 	/** Decrements the lock counter for the invocation list. */
 	inline void UnlockInvocationList( ) const
 	{
-		--InvocationListLockCount;
+		UE_AUTORTFM_OPEN { --InvocationListLockCount; };
+		UE_AUTORTFM_ONABORT(this) { ++InvocationListLockCount; };
 	}
 
 	/** Returns the lock counter for the invocation list. */
