[AutoRTFM] Make Blake3 transactionally safe

Required to fix SOL-7131

diff --git a/Engine/Source/Programs/AutoRTFMTests/Private/Hash.cpp b/Engine/Source/Programs/AutoRTFMTests/Private/Hash.cpp
new file mode 100644
--- /dev/null
+++ b/Engine/Source/Programs/AutoRTFMTests/Private/Hash.cpp
@@ -0,0 +1,143 @@
+// Copyright Epic Games, Inc. All Rights Reserved.
+
+#include "Catch2Includes.h"
+#include "Hash/Blake3.h"
+#include "Memory/CompositeBuffer.h"
+#include "Memory/MemoryFwd.h"
+#include <AutoRTFM/AutoRTFM.h>
+
+TEST_CASE("Hash.Blake3")
+{
+	SECTION("Construct / Destruct")
+	{
+		SECTION("With Abort")
+		{
+			AutoRTFM::ETransactionResult Result = AutoRTFM::Transact([&]
+			{
+				FBlake3 Hash;
+				AutoRTFM::AbortTransaction();
+			});
+			REQUIRE(AutoRTFM::ETransactionResult::AbortedByRequest == Result);
+		}
+		SECTION("With Commit")
+		{
+			AutoRTFM::ETransactionResult Result = AutoRTFM::Transact([&]
+			{
+				FBlake3 Hash;
+			});
+			REQUIRE(AutoRTFM::ETransactionResult::Committed == Result);
+		}
+	}
+
+	SECTION("Update")
+	{
+		const FBlake3Hash ExpectedCommitHash{FUtf8StringView("33f1cb24398ef72a663a8aa5afe2bef9c6d5ff2490e457201c3113d333642627")};
+		const FBlake3Hash ExpectedAbortHash{FUtf8StringView("af1349b9f5f9a1a6a0404dea36dcc9499bcb25c9adc112b7cc9a93cae41f3262")};
+		const char Data[] = "Hello world";
+
+		FBlake3 Hash;
+
+		SECTION("Ptr + Size")
+		{
+			SECTION("With Abort")
+			{
+				AutoRTFM::ETransactionResult Result = AutoRTFM::Transact([&]
+				{
+					Hash.Update(Data, sizeof(Data));
+					AutoRTFM::AbortTransaction();
+				});
+				REQUIRE(AutoRTFM::ETransactionResult::AbortedByRequest == Result);
+				REQUIRE(ExpectedAbortHash == Hash.Finalize());
+			}
+			SECTION("With Commit")
+			{
+				AutoRTFM::ETransactionResult Result = AutoRTFM::Transact([&]
+				{
+					Hash.Update(Data, sizeof(Data));
+				});
+				REQUIRE(AutoRTFM::ETransactionResult::Committed == Result);
+				REQUIRE(ExpectedCommitHash == Hash.Finalize());
+			}
+		}
+
+		SECTION("FMemoryView")
+		{
+			FMemoryView View(Data, sizeof(Data));
+			SECTION("With Abort")
+			{
+				AutoRTFM::ETransactionResult Result = AutoRTFM::Transact([&]
+				{
+					Hash.Update(View);
+					AutoRTFM::AbortTransaction();
+				});
+				REQUIRE(AutoRTFM::ETransactionResult::AbortedByRequest == Result);
+				REQUIRE(ExpectedAbortHash == Hash.Finalize());
+			}
+			SECTION("With Commit")
+			{
+				AutoRTFM::ETransactionResult Result = AutoRTFM::Transact([&]
+				{
+					Hash.Update(View);
+				});
+				REQUIRE(AutoRTFM::ETransactionResult::Committed == Result);
+				REQUIRE(ExpectedCommitHash == Hash.Finalize());
+			}
+		}
+
+		SECTION("FCompositeBuffer")
+		{
+			FCompositeBuffer Buffer{FSharedBuffer::MakeView(Data, sizeof(Data))};
+			SECTION("With Abort")
+			{
+				AutoRTFM::ETransactionResult Result = AutoRTFM::Transact([&]
+				{
+					Hash.Update(Buffer);
+					AutoRTFM::AbortTransaction();
+				});
+				REQUIRE(AutoRTFM::ETransactionResult::AbortedByRequest == Result);
+				REQUIRE(ExpectedAbortHash == Hash.Finalize());
+			}
+			SECTION("With Commit")
+			{
+				AutoRTFM::ETransactionResult Result = AutoRTFM::Transact([&]
+				{
+					Hash.Update(Buffer);
+				});
+				REQUIRE(AutoRTFM::ETransactionResult::Committed == Result);
+				REQUIRE(ExpectedCommitHash == Hash.Finalize());
+			}
+		}
+	}
+
+	SECTION("Finalize")
+	{
+		const FBlake3Hash ExpectedCommitHash{FUtf8StringView("33f1cb24398ef72a663a8aa5afe2bef9c6d5ff2490e457201c3113d333642627")};
+		const FBlake3Hash ExpectedAbortHash{FUtf8StringView("0000000000000000000000000000000000000000000000000000000000000000")};
+		const char Data[] = "Hello world";
+
+		FBlake3 Hash;
+		Hash.Update(Data, sizeof(Data));
+
+		FBlake3Hash Got = FBlake3Hash::Zero;
+
+		SECTION("With Abort")
+		{
+			AutoRTFM::ETransactionResult Result = AutoRTFM::Transact([&]
+			{
+				Got = Hash.Finalize();
+				AutoRTFM::AbortTransaction();
+			});
+			REQUIRE(AutoRTFM::ETransactionResult::AbortedByRequest == Result);
+			REQUIRE(ExpectedAbortHash == Got);
+		}
+		SECTION("With Commit")
+		{
+			AutoRTFM::ETransactionResult Result = AutoRTFM::Transact([&]
+			{
+				Got = Hash.Finalize();
+			});
+			REQUIRE(AutoRTFM::ETransactionResult::Committed == Result);
+			REQUIRE(ExpectedCommitHash == Got);
+		}
+	}
+}
diff --git a/Engine/Source/Runtime/Core/Private/Hash/Blake3.cpp b/Engine/Source/Runtime/Core/Private/Hash/Blake3.cpp
--- a/Engine/Source/Runtime/Core/Private/Hash/Blake3.cpp
+++ b/Engine/Source/Runtime/Core/Private/Hash/Blake3.cpp
@@ -2,6 +2,7 @@
 
 #include "Hash/Blake3.h"
 
+#include "AutoRTFM/AutoRTFM.h"
 #include "Containers/Array.h"
 #include "Containers/ArrayView.h"
 #include "Containers/ContainersFwd.h"
@@ -17,27 +18,47 @@ static_assert(sizeof(FBlake3) == sizeof(blake3_hasher), "Adjust the allocation i
 void FBlake3::Reset()
 {
 	blake3_hasher& Hasher = reinterpret_cast<blake3_hasher&>(HasherBytes);
-	blake3_hasher_init(&Hasher);
+	// blake3 is a precompiled external library without AutoRTFM support.
+	// Call the blake3 function in the open (non-transactionally), and manually
+	// record the update to the Hasher structure.
+	bool bIsClosed = AutoRTFM::IsClosed();
+	UE_AUTORTFM_OPEN
+	{
+		if (bIsClosed)
+		{
+			AutoRTFM::RecordOpenWrite(&Hasher, sizeof(Hasher));
+		}
+		blake3_hasher_init(&Hasher);
+	};
 }
 
 void FBlake3::Update(FMemoryView View)
 {
-	blake3_hasher& Hasher = reinterpret_cast<blake3_hasher&>(HasherBytes);
-	blake3_hasher_update(&Hasher, View.GetData(), View.GetSize());
+	Update(View.GetData(), View.GetSize());
 }
 
 void FBlake3::Update(const void* Data, uint64 Size)
 {
 	blake3_hasher& Hasher = reinterpret_cast<blake3_hasher&>(HasherBytes);
-	blake3_hasher_update(&Hasher, Data, Size);
+	bool bIsClosed = AutoRTFM::IsClosed();
+	// blake3 is a precompiled external library without AutoRTFM support.
+	// Call the blake3 function in the open (non-transactionally), and manually
+	// record the update to the Hasher structure.
+	UE_AUTORTFM_OPEN
+	{
+		if (bIsClosed)
+		{
+			AutoRTFM::RecordOpenWrite(&Hasher, sizeof(Hasher));
+		}
+		blake3_hasher_update(&Hasher, Data, Size);
+	};
 }
 
 void FBlake3::Update(const FCompositeBuffer& Buffer)
 {
-	blake3_hasher& Hasher = reinterpret_cast<blake3_hasher&>(HasherBytes);
 	for (const FSharedBuffer& Segment : Buffer.GetSegments())
 	{
-		blake3_hasher_update(&Hasher, Segment.GetData(), Segment.GetSize());
+		Update(Segment.GetData(), Segment.GetSize());
 	}
 }
 
@@ -47,7 +68,10 @@ FBlake3Hash FBlake3::Finalize() const
 	FBlake3Hash::ByteArray& Output = Hash.GetBytes();
 	static_assert(sizeof(decltype(Output)) == BLAKE3_OUT_LEN, "Mismatch in BLAKE3 hash size.");
 	const blake3_hasher& Hasher = reinterpret_cast<const blake3_hasher&>(HasherBytes);
-	blake3_hasher_finalize(&Hasher, Output, BLAKE3_OUT_LEN);
+	UE_AUTORTFM_OPEN
+	{
+		blake3_hasher_finalize(&Hasher, Output, BLAKE3_OUT_LEN);
+	};
 	return Hash;
 }
 
