[AutoRTFM] Fix macro form of UE_AUTORTFM_TRANSACT.

Previously, this macro was broken and untested. We now have unit tests to verify that the UE_AUTORTFM macro forms of the API work as expected.

diff --git a/Engine/Source/Programs/AutoRTFMTests/Private/API.cpp b/Engine/Source/Programs/AutoRTFMTests/Private/API.cpp
--- a/Engine/Source/Programs/AutoRTFMTests/Private/API.cpp
+++ b/Engine/Source/Programs/AutoRTFMTests/Private/API.cpp
@@ -559,14 +559,41 @@ TEST_CASE("API.IsCommittingOrAborting")
 
 TEST_CASE("API.Transact")
 {
-    int Answer = 6 * 9;
+	int Answer = 6 * 9;
 
-    REQUIRE(AutoRTFM::ETransactionResult::Committed == AutoRTFM::Transact([&]
-    {
-        Answer = 42;
-    }));
+	REQUIRE(AutoRTFM::ETransactionResult::Committed == AutoRTFM::Transact([&]
+	{
+		Answer = 42;
+	}));
 
-    REQUIRE(42 == Answer);
+	REQUIRE(42 == Answer);
+}
+
+TEST_CASE("API.TransactMacro_NoAbort")
+{
+	int Answer = 6 * 9;
+
+	// Allowing the transaction to commit should work.
+	UE_AUTORTFM_TRANSACT
+	{
+		Answer = 42;
+	};
+
+	REQUIRE(42 == Answer);
+}
+
+TEST_CASE("API.TransactMacro_WithAbort")
+{
+	int Answer = 42;
+
+	// Aborting the transaction should also work.
+	UE_AUTORTFM_TRANSACT
+	{
+		Answer = 6 * 9;
+		AutoRTFM::AbortTransaction();
+	};
+
+	REQUIRE(42 == Answer);
 }
 
 TEST_CASE("API.Commit")
@@ -697,6 +724,37 @@ TEST_CASE("API.Open")
     REQUIRE(84 == Answer);
 }
 
+TEST_CASE("API.OpenMacro_NoAbort")
+{
+    int Answer = 6 * 9;
+
+    // An open call outside a transaction succeeds.
+    UE_AUTORTFM_OPEN
+    {
+        Answer = 42;
+    };
+
+    REQUIRE(42 == Answer);
+}
+
+TEST_CASE("API.OpenMacro_WithAbort")
+{
+	int Answer = 21;
+
+	REQUIRE(AutoRTFM::ETransactionResult::AbortedByRequest == AutoRTFM::Transact([&]
+    {
+        // An open call inside a transaction succeeds.
+        UE_AUTORTFM_OPEN
+        {
+            Answer *= 2;
+        };
+
+        AutoRTFM::AbortTransaction();
+    }));
+
+    REQUIRE(42 == Answer);
+}
+
 TEST_CASE("API.Close")
 {
     bool InClosedNest = false;
@@ -819,6 +877,42 @@ TEST_CASE("API.OnCommit")
     REQUIRE(AutoRTFM::ETransactionResult::AbortedByRequest == NestResult);
 }
 
+TEST_CASE("API.OnCommitMacro_NoAbort")
+{
+	int Value = 123;
+
+	UE_AUTORTFM_TRANSACT
+	{
+		UE_AUTORTFM_ONCOMMIT(&Value)
+		{
+			Value = 456;
+		};
+
+		Value = 789;
+	};
+
+	REQUIRE(Value == 456);
+}
+
+TEST_CASE("API.OnCommitMacro_WithAbort")
+{
+	int Value = 123;
+
+	UE_AUTORTFM_TRANSACT
+	{
+		UE_AUTORTFM_ONCOMMIT(&Value)
+		{
+			Value = 456;
+		};
+
+		Value = 789;
+
+		AutoRTFM::AbortTransaction();
+	};
+
+	REQUIRE(Value == 123);
+}
+
 TEST_CASE("API.OnAbort")
 {
 	// Too hard to get this test working when retrying nested transactions so bail!
@@ -908,6 +1002,42 @@ TEST_CASE("API.OnAbort")
     REQUIRE(AutoRTFM::ETransactionResult::AbortedByRequest == NestResult);
 }
 
+TEST_CASE("API.OnAbortMacro_NoAbort")
+{
+	int Value = 123;
+
+	UE_AUTORTFM_TRANSACT
+	{
+		Value = 456;
+
+		UE_AUTORTFM_ONABORT(&Value)
+		{
+			Value = 123;
+		};
+	};
+
+	REQUIRE(Value == 456);
+}
+
+TEST_CASE("API.OnAbortMacro_WithAbort")
+{
+	int Value = 123;
+
+	UE_AUTORTFM_TRANSACT
+	{
+		Value = 234;
+
+		UE_AUTORTFM_ONABORT(&Value)
+		{
+			Value = 123;
+		};
+
+		AutoRTFM::AbortTransaction();
+	};
+
+	REQUIRE(Value == 123);
+}
+
 TEST_CASE("API.DidAllocate")
 {
     constexpr unsigned Size = 1024;
diff --git a/Engine/Source/Runtime/Core/Public/AutoRTFM/AutoRTFM.h b/Engine/Source/Runtime/Core/Public/AutoRTFM/AutoRTFM.h
--- a/Engine/Source/Runtime/Core/Public/AutoRTFM/AutoRTFM.h
+++ b/Engine/Source/Runtime/Core/Public/AutoRTFM/AutoRTFM.h
@@ -1060,7 +1060,7 @@ namespace AutoRTFM::Private
 #define UE_AUTORTFM_OPEN_IMPL          ::AutoRTFM::Private::FOpenHelper{} + [&]()
 #define UE_AUTORTFM_ONABORT_IMPL(...)  ::AutoRTFM::Private::FOnAbortHelper{} + [__VA_ARGS__]()
 #define UE_AUTORTFM_ONCOMMIT_IMPL(...) ::AutoRTFM::Private::FOnCommitHelper{} + [__VA_ARGS__]()
-#define UE_AUTORTFM_TRANSACT_IMPL      ::AutoRTFM::Private::FTransactHelper{} + [&]() { __VA_ARGS__ })
+#define UE_AUTORTFM_TRANSACT_IMPL      ::AutoRTFM::Private::FTransactHelper{} + [&]()
 #else
 
 // Do nothing, these should be followed by blocks that should be either executed or not executed
