Disable benchmarks by default for AutoRTFMTests.

diff --git a/Engine/Source/Programs/AutoRTFMTests/Private/AutoRTFMTests.cpp b/Engine/Source/Programs/AutoRTFMTests/Private/AutoRTFMTests.cpp
--- a/Engine/Source/Programs/AutoRTFMTests/Private/AutoRTFMTests.cpp
+++ b/Engine/Source/Programs/AutoRTFMTests/Private/AutoRTFMTests.cpp
@@ -49,27 +49,36 @@ int RunTests(int ArgC, const char* ArgV[])
 
 	Catch::Session Session;
 
-	bool NoRetry = false;
-	bool RetryNestedToo = false;
+	bool bNoRetry = false;
+	bool bRetryNestedToo = false;
+	bool bEnableBenchmarks = false;
 
 	Session.cli(Session.cli()
-		| Catch::Clara::Opt(NoRetry)["--no-retry"]
-		| Catch::Clara::Opt(RetryNestedToo)["--retry-nested-too"]);
+		| Catch::Clara::Opt(bNoRetry)["--no-retry"]
+		| Catch::Clara::Opt(bRetryNestedToo)["--retry-nested-too"]
+		| Catch::Clara::Opt(bEnableBenchmarks)["--enable-benchmarks"]);
 
 	{
 		const int Result = Session.applyCommandLine(ArgC, ArgV);
 
+		if (!bEnableBenchmarks)
+		{
+			constexpr int FakeArgC = 2;
+			const char* const FakeArgV[FakeArgC] = { ArgV[0], "--skip-benchmarks" };
+			Session.applyCommandLine(FakeArgC, FakeArgV);
+		}
+
 		if (0 != Result)
 		{
 			return Result;
 		}
 	}
 
-	if (RetryNestedToo)
+	if (bRetryNestedToo)
 	{
 		AutoRTFM::ForTheRuntime::SetRetryTransaction(AutoRTFM::ForTheRuntime::EAutoRTFMRetryTransactionState::RetryNestedToo);
 	}
-	else if (NoRetry)
+	else if (bNoRetry)
 	{
 		AutoRTFM::ForTheRuntime::SetRetryTransaction(AutoRTFM::ForTheRuntime::EAutoRTFMRetryTransactionState::NoRetry);
 	}
