[AutoRTFM] Avoid using Catch2 macros in both the open and closed

REQUIRE() performs memory writes even when the expression passes. Attempting to mix these can corrupt the internal Catch2 state.

diff --git a/Engine/Source/Programs/AutoRTFMTests/Private/API.cpp b/Engine/Source/Programs/AutoRTFMTests/Private/API.cpp
--- a/Engine/Source/Programs/AutoRTFMTests/Private/API.cpp
+++ b/Engine/Source/Programs/AutoRTFMTests/Private/API.cpp
@@ -761,27 +761,28 @@ TEST_CASE("API.Close")
     bool InOpenNest = false;
     bool InClosedNestInOpenNest = false;
 
-    REQUIRE(AutoRTFM::ETransactionResult::AbortedByRequest == AutoRTFM::Transact([&]
+    AutoRTFM::ETransactionResult Result = AutoRTFM::Transact([&]
     {
         // A closed call inside a transaction does not abort.
-        REQUIRE(AutoRTFM::EContextStatus::OnTrack == AutoRTFM::Close([&]
-        {
-            InClosedNest = true;
-        }));
+        AutoRTFM::EContextStatus CloseStatusA = AutoRTFM::Close([&] { InClosedNest = true; });
+        REQUIRE(AutoRTFM::EContextStatus::OnTrack == CloseStatusA);
 
+        AutoRTFM::EContextStatus CloseStatusB{AutoRTFM::EContextStatus::Idle};
         AutoRTFM::Open([&]
         {
-            // A closed call inside an open does not abort either.
-			REQUIRE(AutoRTFM::EContextStatus::OnTrack == AutoRTFM::Close([&]
-            {
-                InClosedNestInOpenNest = true;
-            }));
-
+            CloseStatusB = AutoRTFM::Close([&] { InClosedNestInOpenNest = true; });
             InOpenNest = true;
         });
 
+        // A closed call inside an open does not abort either.
+        // Note: REQUIRE() must be called outside of the Open() lambda as the Catch2 framework may
+        // perform writes which are being tracked by the closed transaction.
+        REQUIRE(AutoRTFM::EContextStatus::OnTrack == CloseStatusB);
+
         AutoRTFM::AbortTransaction();
-    }));
+    });
+
+    REQUIRE(AutoRTFM::ETransactionResult::AbortedByRequest == Result);
 
     REQUIRE(false == InClosedNest);
     REQUIRE(true == InOpenNest);
diff --git a/Engine/Source/Programs/AutoRTFMTests/Private/OpenAPI.cpp b/Engine/Source/Programs/AutoRTFMTests/Private/OpenAPI.cpp
--- a/Engine/Source/Programs/AutoRTFMTests/Private/OpenAPI.cpp
+++ b/Engine/Source/Programs/AutoRTFMTests/Private/OpenAPI.cpp
@@ -476,24 +476,32 @@ TEST_CASE("OpenAPI.StackWriteCommitInTheOpen1")
 
 TEST_CASE("OpenAPI.StackWriteCommitInTheOpen2")
 {
+	AutoRTFM::EContextStatus Status = AutoRTFM::EContextStatus::Idle;
+
 	int Value = 0;
+
 	AutoRTFM::Transact([&]()
 	{
 		AutoRTFM::Open([&]()
 		{
 			AutoRTFM::ForTheRuntime::StartTransaction();
-			REQUIRE(AutoRTFM::EContextStatus::OnTrack == AutoRTFM::Close([&]()
+
+			Status = AutoRTFM::Close([&]()
 			{
 				AutoRTFM::Open([&]()
 				{
 					AutoRTFM::RecordOpenWrite(&Value);
 					Value = 10;
 				});
-			}));
+			});
 
 			AutoRTFM::ForTheRuntime::CommitTransaction();
-			REQUIRE(Value == 10);
 		});
+
+		// Note: REQUIRE() must be called outside of the Open() lambda as the Catch2 framework may
+		// perfore writes which are being tracked by the closed transaction.
+		REQUIRE(AutoRTFM::EContextStatus::OnTrack == Status);
+		REQUIRE(Value == 10);
 	});
 }
 
@@ -543,8 +551,8 @@ TEST_CASE("OpenAPI.WriteMemory1")
 
 			AutoRTFM::ForTheRuntime::WriteMemory(&value1, &sourceValue);
 			AutoRTFM::ForTheRuntime::CommitTransaction();
-			REQUIRE(value1 == 10);
 		});
+		REQUIRE(value1 == 10);
 	});
 }
 
