Fix the MTTransactionallySafeAccessDetector to work across open/closed and in oncommit/onabort too.

diff --git a/Engine/Source/Programs/AutoRTFMTests/Private/MTTSAccessDetector.cpp b/Engine/Source/Programs/AutoRTFMTests/Private/MTTSAccessDetector.cpp
--- a/Engine/Source/Programs/AutoRTFMTests/Private/MTTSAccessDetector.cpp
+++ b/Engine/Source/Programs/AutoRTFMTests/Private/MTTSAccessDetector.cpp
@@ -3,447 +3,899 @@
 #include "Catch2Includes.h"
 #include "Misc/MTTransactionallySafeAccessDetector.h"
 #include <AutoRTFM/AutoRTFM.h>
+#include <atomic>
+#include <thread>
 
 #if ENABLE_MT_DETECTOR
 
 TEST_CASE("MTTransactionallySafeAccessDetector")
 {
-	FRWTransactionallySafeAccessDetector Detector;
-
-	SECTION("AcquireWriteAccess, ReleaseWriteAccess")
+	SECTION("Transact(UE_MT_SCOPED_READ_ACCESS, Scope(UE_MT_SCOPED_READ_ACCESS), Abort)")
 	{
-		Detector.AcquireWriteAccess();
-		Detector.ReleaseWriteAccess();
+		FRWTransactionallySafeAccessDetector Detector;
+		AutoRTFM::ETransactionResult Result = AutoRTFM::Transact([&]
+			{
+				UE_MT_SCOPED_READ_ACCESS(Detector);
+				{
+					UE_MT_SCOPED_READ_ACCESS(Detector);
+				}
+				AutoRTFM::AbortTransaction();
+			});
+		REQUIRE(AutoRTFM::ETransactionResult::AbortedByRequest == Result);
 	}
 
-	SECTION("AcquireReadAccess, ReleaseReadAccess")
+	SECTION("UE_MT_SCOPED_READ_ACCESS, Transact(UE_MT_SCOPED_READ_ACCESS)")
 	{
-		Detector.AcquireReadAccess();
-		Detector.ReleaseReadAccess();
+		FRWTransactionallySafeAccessDetector Detector;
+		UE_MT_SCOPED_READ_ACCESS(Detector);
+		AutoRTFM::ETransactionResult Result = AutoRTFM::Transact([&]
+			{
+				UE_MT_SCOPED_READ_ACCESS(Detector);
+			});
+		REQUIRE(AutoRTFM::ETransactionResult::Committed == Result);
 	}
 
-	SECTION("AcquireReadAccess, AcquireReadAccess, ReleaseReadAccess, ReleaseReadAccess")
+	SECTION("UE_MT_SCOPED_READ_ACCESS, Transact(UE_MT_SCOPED_READ_ACCESS, Abort)")
 	{
-		Detector.AcquireReadAccess();
-		Detector.AcquireReadAccess();
-		Detector.ReleaseReadAccess();
-		Detector.ReleaseReadAccess();
+		FRWTransactionallySafeAccessDetector Detector;
+		UE_MT_SCOPED_READ_ACCESS(Detector);
+		AutoRTFM::ETransactionResult Result = AutoRTFM::Transact([&]
+			{
+				UE_MT_SCOPED_READ_ACCESS(Detector);
+				AutoRTFM::AbortTransaction();
+			});
+		REQUIRE(AutoRTFM::ETransactionResult::AbortedByRequest == Result);
 	}
 
-	SECTION("Transact(AcquireWriteAccess, ReleaseWriteAccess)")
+	SECTION("Transact(UE_MT_SCOPED_WRITE_ACCESS)")
 	{
+		FRWTransactionallySafeAccessDetector Detector;
 		AutoRTFM::ETransactionResult Result = AutoRTFM::Transact([&]
-		{
-			Detector.AcquireWriteAccess();
-			Detector.ReleaseWriteAccess();
-		});
+			{
+				UE_MT_SCOPED_WRITE_ACCESS(Detector);
+			});
 		REQUIRE(AutoRTFM::ETransactionResult::Committed == Result);
 	}
 
-	SECTION("Transact(AcquireWriteAccess, ReleaseWriteAccess, Abort)")
+	SECTION("Transact(UE_MT_SCOPED_WRITE_ACCESS, Abort)")
 	{
+		FRWTransactionallySafeAccessDetector Detector;
 		AutoRTFM::ETransactionResult Result = AutoRTFM::Transact([&]
-		{
-			Detector.AcquireWriteAccess();
-			Detector.ReleaseWriteAccess();
-			AutoRTFM::AbortTransaction();
-		});
+			{
+				UE_MT_SCOPED_WRITE_ACCESS(Detector);
+				AutoRTFM::AbortTransaction();
+			});
 		REQUIRE(AutoRTFM::ETransactionResult::AbortedByRequest == Result);
 	}
 
-	SECTION("Transact(AcquireWriteAccess), ReleaseWriteAccess")
+	SECTION("Transact(Scope(UE_MT_SCOPED_WRITE_ACCESS), Abort)")
+	{
+		FRWTransactionallySafeAccessDetector Detector;
+		AutoRTFM::ETransactionResult Result = AutoRTFM::Transact([&]
+			{
+				{
+					UE_MT_SCOPED_WRITE_ACCESS(Detector);
+				}
+				AutoRTFM::AbortTransaction();
+			});
+		REQUIRE(AutoRTFM::ETransactionResult::AbortedByRequest == Result);
+	}
+
+	SECTION("Transact(Scope(UE_MT_SCOPED_READ_ACCESS), Scope(UE_MT_SCOPED_READ_ACCESS))")
+	{
+		FRWTransactionallySafeAccessDetector Detector;
+		AutoRTFM::ETransactionResult Result = AutoRTFM::Transact([&]
+			{
+				{
+					UE_MT_SCOPED_READ_ACCESS(Detector);
+				}
+				{
+					UE_MT_SCOPED_READ_ACCESS(Detector);
+				}
+			});
+		REQUIRE(AutoRTFM::ETransactionResult::Committed == Result);
+	}
+
+	SECTION("Transact(Scope(UE_MT_SCOPED_READ_ACCESS), Scope(UE_MT_SCOPED_WRITE_ACCESS))")
+	{
+		FRWTransactionallySafeAccessDetector Detector;
+		AutoRTFM::ETransactionResult Result = AutoRTFM::Transact([&]
+			{
+				{
+					UE_MT_SCOPED_READ_ACCESS(Detector);
+				}
+				{
+					UE_MT_SCOPED_WRITE_ACCESS(Detector);
+				}
+			});
+		REQUIRE(AutoRTFM::ETransactionResult::Committed == Result);
+	}
+
+	SECTION("Transact(Scope(UE_MT_SCOPED_WRITE_ACCESS), Scope(UE_MT_SCOPED_READ_ACCESS))")
 	{
+		FRWTransactionallySafeAccessDetector Detector;
 		AutoRTFM::ETransactionResult Result = AutoRTFM::Transact([&]
-		{
-			Detector.AcquireWriteAccess();
-		});
+			{
+				{
+					UE_MT_SCOPED_WRITE_ACCESS(Detector);
+				}
+				{
+					UE_MT_SCOPED_READ_ACCESS(Detector);
+				}
+			});
 		REQUIRE(AutoRTFM::ETransactionResult::Committed == Result);
+	}
 
-		Detector.ReleaseWriteAccess();
+	SECTION("Transact(Scope(UE_MT_SCOPED_WRITE_ACCESS), Scope(UE_MT_SCOPED_WRITE_ACCESS))")
+	{
+		FRWTransactionallySafeAccessDetector Detector;
+		AutoRTFM::ETransactionResult Result = AutoRTFM::Transact([&]
+			{
+				{
+					UE_MT_SCOPED_WRITE_ACCESS(Detector);
+				}
+				{
+					UE_MT_SCOPED_WRITE_ACCESS(Detector);
+				}
+			});
+		REQUIRE(AutoRTFM::ETransactionResult::Committed == Result);
 	}
 
-	SECTION("Transact(AcquireWriteAccess, Abort)")
+	SECTION("Transact(Scope(UE_MT_SCOPED_READ_ACCESS), Scope(UE_MT_SCOPED_READ_ACCESS), Abort)")
 	{
+		FRWTransactionallySafeAccessDetector Detector;
 		AutoRTFM::ETransactionResult Result = AutoRTFM::Transact([&]
-		{
-			Detector.AcquireWriteAccess();
-			AutoRTFM::AbortTransaction();
-		});
+			{
+				{
+					UE_MT_SCOPED_READ_ACCESS(Detector);
+				}
+				{
+					UE_MT_SCOPED_READ_ACCESS(Detector);
+				}
+				AutoRTFM::AbortTransaction();
+			});
 		REQUIRE(AutoRTFM::ETransactionResult::AbortedByRequest == Result);
 	}
 
-	SECTION("Transact(AcquireReadAccess, Abort)")
+	SECTION("Transact(Scope(UE_MT_SCOPED_READ_ACCESS), Scope(UE_MT_SCOPED_WRITE_ACCESS), Abort)")
 	{
+		FRWTransactionallySafeAccessDetector Detector;
 		AutoRTFM::ETransactionResult Result = AutoRTFM::Transact([&]
-		{
-			Detector.AcquireReadAccess();
-			AutoRTFM::AbortTransaction();
-		});
+			{
+				{
+					UE_MT_SCOPED_READ_ACCESS(Detector);
+				}
+				{
+					UE_MT_SCOPED_WRITE_ACCESS(Detector);
+				}
+				AutoRTFM::AbortTransaction();
+			});
 		REQUIRE(AutoRTFM::ETransactionResult::AbortedByRequest == Result);
 	}
 
-	SECTION("AcquireReadAccess, Transact(AcquireReadAccess, Abort), ReleaseReadAccess")
+	SECTION("Transact(Scope(UE_MT_SCOPED_WRITE_ACCESS), Scope(UE_MT_SCOPED_READ_ACCESS), Abort)")
 	{
-		Detector.AcquireReadAccess();
+		FRWTransactionallySafeAccessDetector Detector;
+		AutoRTFM::ETransactionResult Result = AutoRTFM::Transact([&]
+			{
+				{
+					UE_MT_SCOPED_WRITE_ACCESS(Detector);
+				}
+				{
+					UE_MT_SCOPED_READ_ACCESS(Detector);
+				}
+				AutoRTFM::AbortTransaction();
+			});
+		REQUIRE(AutoRTFM::ETransactionResult::AbortedByRequest == Result);
+	}
 
+	SECTION("Transact(Scope(UE_MT_SCOPED_WRITE_ACCESS), Scope(UE_MT_SCOPED_WRITE_ACCESS), Abort)")
+	{
+		FRWTransactionallySafeAccessDetector Detector;
 		AutoRTFM::ETransactionResult Result = AutoRTFM::Transact([&]
-		{
-			Detector.AcquireReadAccess();
-			AutoRTFM::AbortTransaction();
-		});
+			{
+				{
+					UE_MT_SCOPED_WRITE_ACCESS(Detector);
+				}
+				{
+					UE_MT_SCOPED_WRITE_ACCESS(Detector);
+				}
+				AutoRTFM::AbortTransaction();
+			});
 		REQUIRE(AutoRTFM::ETransactionResult::AbortedByRequest == Result);
+	}
 
-		Detector.ReleaseReadAccess();
+	SECTION("Transact(Scope(UE_MT_SCOPED_READ_ACCESS), Scope(UE_MT_SCOPED_READ_ACCESS, Abort))")
+	{
+		FRWTransactionallySafeAccessDetector Detector;
+		AutoRTFM::ETransactionResult Result = AutoRTFM::Transact([&]
+			{
+				{
+					UE_MT_SCOPED_READ_ACCESS(Detector);
+				}
+				{
+					UE_MT_SCOPED_READ_ACCESS(Detector);
+					AutoRTFM::AbortTransaction();
+				}
+			});
+		REQUIRE(AutoRTFM::ETransactionResult::AbortedByRequest == Result);
 	}
 
-	SECTION("AcquireWriteAccess, Transact(ReleaseWriteAccess)")
+	SECTION("Transact(Scope(UE_MT_SCOPED_READ_ACCESS), Scope(UE_MT_SCOPED_WRITE_ACCESS, Abort))")
 	{
-		Detector.AcquireWriteAccess();
+		FRWTransactionallySafeAccessDetector Detector;
+		AutoRTFM::ETransactionResult Result = AutoRTFM::Transact([&]
+			{
+				{
+					UE_MT_SCOPED_READ_ACCESS(Detector);
+				}
+				{
+					UE_MT_SCOPED_WRITE_ACCESS(Detector);
+					AutoRTFM::AbortTransaction();
+				}
+			});
+		REQUIRE(AutoRTFM::ETransactionResult::AbortedByRequest == Result);
+	}
 
+	SECTION("Transact(Scope(UE_MT_SCOPED_WRITE_ACCESS), Scope(UE_MT_SCOPED_READ_ACCESS, Abort))")
+	{
+		FRWTransactionallySafeAccessDetector Detector;
 		AutoRTFM::ETransactionResult Result = AutoRTFM::Transact([&]
-		{
-			Detector.ReleaseWriteAccess();
-		});
-		REQUIRE(AutoRTFM::ETransactionResult::Committed == Result);
+			{
+				{
+					UE_MT_SCOPED_WRITE_ACCESS(Detector);
+				}
+				{
+					UE_MT_SCOPED_READ_ACCESS(Detector);
+					AutoRTFM::AbortTransaction();
+				}
+			});
+		REQUIRE(AutoRTFM::ETransactionResult::AbortedByRequest == Result);
 	}
 
-	SECTION("AcquireReadAccess, Transact(ReleaseReadAccess)")
+	SECTION("Transact(Scope(UE_MT_SCOPED_WRITE_ACCESS), Scope(UE_MT_SCOPED_WRITE_ACCESS, Abort))")
 	{
-		Detector.AcquireReadAccess();
+		FRWTransactionallySafeAccessDetector Detector;
+		AutoRTFM::ETransactionResult Result = AutoRTFM::Transact([&]
+			{
+				{
+					UE_MT_SCOPED_WRITE_ACCESS(Detector);
+				}
+				{
+					UE_MT_SCOPED_WRITE_ACCESS(Detector);
+					AutoRTFM::AbortTransaction();
+				}
+			});
+		REQUIRE(AutoRTFM::ETransactionResult::AbortedByRequest == Result);
+	}
 
+	SECTION("Transact(UE_MT_SCOPED_READ_ACCESS, Open(UE_MT_SCOPED_READ_ACCESS))")
+	{
+		FRWTransactionallySafeAccessDetector Detector;
 		AutoRTFM::ETransactionResult Result = AutoRTFM::Transact([&]
-		{
-			Detector.ReleaseReadAccess();
-		});
+			{
+				UE_MT_SCOPED_READ_ACCESS(Detector);
+				AutoRTFM::Open([&]
+					{
+						UE_MT_SCOPED_READ_ACCESS(Detector);
+					});
+			});
 		REQUIRE(AutoRTFM::ETransactionResult::Committed == Result);
 	}
 
-	SECTION("AcquireWriteAccess, Transact(ReleaseWriteAccess, Abort), ReleaseWriteAccess")
+	SECTION("Transact(UE_MT_SCOPED_READ_ACCESS, Open(UE_MT_SCOPED_READ_ACCESS), Abort)")
 	{
-		Detector.AcquireWriteAccess();
-
+		FRWTransactionallySafeAccessDetector Detector;
 		AutoRTFM::ETransactionResult Result = AutoRTFM::Transact([&]
-		{
-			Detector.ReleaseWriteAccess();
-			AutoRTFM::AbortTransaction();
-		});
+			{
+				UE_MT_SCOPED_READ_ACCESS(Detector);
+				AutoRTFM::Open([&]
+					{
+						UE_MT_SCOPED_READ_ACCESS(Detector);
+					});
+				AutoRTFM::AbortTransaction();
+			});
 		REQUIRE(AutoRTFM::ETransactionResult::AbortedByRequest == Result);
+	}
 
-		Detector.ReleaseWriteAccess();
+	SECTION("Transact(Scope(UE_MT_SCOPED_READ_ACCESS, Open(UE_MT_SCOPED_READ_ACCESS)), Abort)")
+	{
+		FRWTransactionallySafeAccessDetector Detector;
+		AutoRTFM::ETransactionResult Result = AutoRTFM::Transact([&]
+			{
+				{
+					UE_MT_SCOPED_READ_ACCESS(Detector);
+					AutoRTFM::Open([&]
+						{
+							UE_MT_SCOPED_READ_ACCESS(Detector);
+						});
+				}
+				AutoRTFM::AbortTransaction();
+			});
+		REQUIRE(AutoRTFM::ETransactionResult::AbortedByRequest == Result);
 	}
 
-	SECTION("Transact(AcquireReadAccess, AcquireReadAccess, ReleaseReadAccess, ReleaseReadAccess)")
+	SECTION("Transact(Scope(UE_MT_SCOPED_READ_ACCESS), Open(UE_MT_SCOPED_READ_ACCESS))")
 	{
+		FRWTransactionallySafeAccessDetector Detector;
 		AutoRTFM::ETransactionResult Result = AutoRTFM::Transact([&]
-		{
-			Detector.AcquireReadAccess();
-			Detector.AcquireReadAccess();
-			Detector.ReleaseReadAccess();
-			Detector.ReleaseReadAccess();
-		});
+			{
+				{
+					UE_MT_SCOPED_READ_ACCESS(Detector);
+				}
+				AutoRTFM::Open([&]
+					{
+						UE_MT_SCOPED_READ_ACCESS(Detector);
+					});
+			});
 		REQUIRE(AutoRTFM::ETransactionResult::Committed == Result);
 	}
 
-	SECTION("Transact(AcquireReadAccess, AcquireReadAccess, ReleaseReadAccess, ReleaseReadAccess, Abort)")
+	SECTION("Transact(Scope(UE_MT_SCOPED_READ_ACCESS), Open(UE_MT_SCOPED_WRITE_ACCESS))")
 	{
+		FRWTransactionallySafeAccessDetector Detector;
 		AutoRTFM::ETransactionResult Result = AutoRTFM::Transact([&]
-		{
-			Detector.AcquireReadAccess();
-			Detector.AcquireReadAccess();
-			Detector.ReleaseReadAccess();
-			Detector.ReleaseReadAccess();
-			AutoRTFM::AbortTransaction();
-		});
-		REQUIRE(AutoRTFM::ETransactionResult::AbortedByRequest == Result);
+			{
+				{
+					UE_MT_SCOPED_READ_ACCESS(Detector);
+				}
+				AutoRTFM::Open([&]
+					{
+						UE_MT_SCOPED_WRITE_ACCESS(Detector);
+					});
+			});
+		REQUIRE(AutoRTFM::ETransactionResult::Committed == Result);
 	}
-	
-	SECTION("Transact(AcquireWriteAccess, ReleaseWriteAccess, AcquireReadAccess)")
+
+	SECTION("Transact(Scope(UE_MT_SCOPED_WRITE_ACCESS), Open(UE_MT_SCOPED_READ_ACCESS))")
 	{
+		FRWTransactionallySafeAccessDetector Detector;
 		AutoRTFM::ETransactionResult Result = AutoRTFM::Transact([&]
-		{
-			Detector.AcquireWriteAccess();
-			Detector.ReleaseWriteAccess();
-			Detector.AcquireReadAccess();
-		});
+			{
+				{
+					UE_MT_SCOPED_WRITE_ACCESS(Detector);
+				}
+				AutoRTFM::Open([&]
+					{
+						UE_MT_SCOPED_READ_ACCESS(Detector);
+					});
+			});
 		REQUIRE(AutoRTFM::ETransactionResult::Committed == Result);
-		
-		Detector.ReleaseReadAccess();
 	}
 
-	SECTION("Transact(AcquireWriteAccess, ReleaseWriteAccess, AcquireReadAccess, Abort)")
+	SECTION("Transact(Scope(UE_MT_SCOPED_WRITE_ACCESS), Open(UE_MT_SCOPED_WRITE_ACCESS))")
 	{
+		FRWTransactionallySafeAccessDetector Detector;
 		AutoRTFM::ETransactionResult Result = AutoRTFM::Transact([&]
-		{
-			Detector.AcquireWriteAccess();
-			Detector.ReleaseWriteAccess();
-			Detector.AcquireReadAccess();
-			AutoRTFM::AbortTransaction();
-		});
-		REQUIRE(AutoRTFM::ETransactionResult::AbortedByRequest == Result);
+			{
+				{
+					UE_MT_SCOPED_WRITE_ACCESS(Detector);
+				}
+				AutoRTFM::Open([&]
+					{
+						UE_MT_SCOPED_WRITE_ACCESS(Detector);
+					});
+			});
+		REQUIRE(AutoRTFM::ETransactionResult::Committed == Result);
 	}
 
-	SECTION("Transact(AcquireWriteAccess, ReleaseWriteAccess, AcquireReadAccess, ReleaseReadAccess)")
+	SECTION("Transact(Open(UE_MT_SCOPED_READ_ACCESS), Scope(UE_MT_SCOPED_READ_ACCESS))")
 	{
+		FRWTransactionallySafeAccessDetector Detector;
 		AutoRTFM::ETransactionResult Result = AutoRTFM::Transact([&]
-		{
-			Detector.AcquireWriteAccess();
-			Detector.ReleaseWriteAccess();
-			Detector.AcquireReadAccess();
-			Detector.ReleaseReadAccess();
-		});
+			{
+				AutoRTFM::Open([&]
+					{
+						UE_MT_SCOPED_READ_ACCESS(Detector);
+					});
+				UE_MT_SCOPED_READ_ACCESS(Detector);
+			});
 		REQUIRE(AutoRTFM::ETransactionResult::Committed == Result);
 	}
 
-	SECTION("Transact(AcquireWriteAccess, ReleaseWriteAccess, AcquireReadAccess, ReleaseReadAccess, Abort)")
+	SECTION("Transact(Open(UE_MT_SCOPED_READ_ACCESS), Scope(UE_MT_SCOPED_WRITE_ACCESS))")
 	{
+		FRWTransactionallySafeAccessDetector Detector;
 		AutoRTFM::ETransactionResult Result = AutoRTFM::Transact([&]
-		{
-			Detector.AcquireWriteAccess();
-			Detector.ReleaseWriteAccess();
-			Detector.AcquireReadAccess();
-			Detector.ReleaseReadAccess();
-			AutoRTFM::AbortTransaction();
-		});
-		REQUIRE(AutoRTFM::ETransactionResult::AbortedByRequest == Result);
+			{
+				AutoRTFM::Open([&]
+					{
+						UE_MT_SCOPED_READ_ACCESS(Detector);
+					});
+				UE_MT_SCOPED_WRITE_ACCESS(Detector);
+			});
+		REQUIRE(AutoRTFM::ETransactionResult::Committed == Result);
 	}
 
-	SECTION("AcquireReadAccess, Transact(ReleaseReadAccess, AcquireWriteAccess), ReleaseWriteAccess")
+	SECTION("Transact(Open(UE_MT_SCOPED_WRITE_ACCESS), Scope(UE_MT_SCOPED_READ_ACCESS))")
 	{
-		Detector.AcquireReadAccess();
+		FRWTransactionallySafeAccessDetector Detector;
+		AutoRTFM::ETransactionResult Result = AutoRTFM::Transact([&]
+			{
+				AutoRTFM::Open([&]
+					{
+						UE_MT_SCOPED_WRITE_ACCESS(Detector);
+					});
+				UE_MT_SCOPED_READ_ACCESS(Detector);
+			});
+		REQUIRE(AutoRTFM::ETransactionResult::Committed == Result);
+	}
 
+	SECTION("Transact(Open(UE_MT_SCOPED_WRITE_ACCESS), UE_MT_SCOPED_WRITE_ACCESS)")
+	{
+		FRWTransactionallySafeAccessDetector Detector;
 		AutoRTFM::ETransactionResult Result = AutoRTFM::Transact([&]
-		{
-			Detector.ReleaseReadAccess();
-			Detector.AcquireWriteAccess();
-		});
+			{
+				AutoRTFM::Open([&]
+					{
+						UE_MT_SCOPED_WRITE_ACCESS(Detector);
+					});
+				UE_MT_SCOPED_WRITE_ACCESS(Detector);
+			});
 		REQUIRE(AutoRTFM::ETransactionResult::Committed == Result);
+	}
+
+	SECTION("Transact(Scope(UE_MT_SCOPED_READ_ACCESS, Spawn(UE_MT_SCOPED_READ_ACCESS)))")
+	{
+		FRWTransactionallySafeAccessDetector Detector;
+		std::atomic_uint Handshake = 0; // Used to order this thread and the spawnee.
+		AutoRTFM::ETransactionResult Result = AutoRTFM::Transact([&]
+			{
+				// Because we can retry transactions reset handshake to zero.
+				AutoRTFM::Open([&]
+					{
+						Handshake = 0;
+					});
+
+				{
+					UE_MT_SCOPED_READ_ACCESS(Detector);
+
+					AutoRTFM::Open([&]
+						{
+							std::thread([&]
+								{
+									// Let the main thread progress.
+									Handshake += 1;
+
+									UE_MT_SCOPED_READ_ACCESS(Detector);
 
-		Detector.ReleaseWriteAccess();
+									// Wait for the main thread.
+									while (2 != Handshake) {}
+								}).detach();
+						});
+				}
+
+				AutoRTFM::Open([&]
+					{
+						// Wait for the spawnee.
+						while (1 != Handshake) {}
+
+						// Let the spawnee progress.
+						Handshake += 1;
+					});
+			});
+		REQUIRE(AutoRTFM::ETransactionResult::Committed == Result);
 	}
 
-	SECTION("AcquireReadAccess, Transact(ReleaseReadAccess, AcquireWriteAccess, Abort), ReleaseReadAccess")
+	SECTION("Transact(UE_MT_SCOPED_READ_ACCESS, Transact(UE_MT_SCOPED_READ_ACCESS))")
 	{
-		Detector.AcquireReadAccess();
+		FRWTransactionallySafeAccessDetector Detector;
+		AutoRTFM::ETransactionResult Result = AutoRTFM::Transact([&]
+			{
+				UE_MT_SCOPED_READ_ACCESS(Detector);
+
+				AutoRTFM::Commit([&]
+					{
+						UE_MT_SCOPED_READ_ACCESS(Detector);
+					});
+			});
+		REQUIRE(AutoRTFM::ETransactionResult::Committed == Result);
+	}
 
+	SECTION("Transact(Scope(UE_MT_SCOPED_READ_ACCESS), Transact(UE_MT_SCOPED_READ_ACCESS))")
+	{
+		FRWTransactionallySafeAccessDetector Detector;
 		AutoRTFM::ETransactionResult Result = AutoRTFM::Transact([&]
-		{
-			Detector.ReleaseReadAccess();
-			Detector.AcquireWriteAccess();
-			AutoRTFM::AbortTransaction();
-		});
-		REQUIRE(AutoRTFM::ETransactionResult::AbortedByRequest == Result);
+			{
+				{
+					UE_MT_SCOPED_READ_ACCESS(Detector);
+				}
 
-		Detector.ReleaseReadAccess();
+				AutoRTFM::Commit([&]
+					{
+						UE_MT_SCOPED_READ_ACCESS(Detector);
+					});
+			});
+		REQUIRE(AutoRTFM::ETransactionResult::Committed == Result);
 	}
 
-	SECTION("AcquireReadAccess, Transact(ReleaseReadAccess, AcquireWriteAccess, ReleaseWriteAccess)")
+	SECTION("Transact(Scope(UE_MT_SCOPED_READ_ACCESS), Transact(UE_MT_SCOPED_WRITE_ACCESS))")
 	{
-		Detector.AcquireReadAccess();
+		FRWTransactionallySafeAccessDetector Detector;
+		AutoRTFM::ETransactionResult Result = AutoRTFM::Transact([&]
+			{
+				{
+					UE_MT_SCOPED_READ_ACCESS(Detector);
+				}
 
+				AutoRTFM::Commit([&]
+					{
+						UE_MT_SCOPED_WRITE_ACCESS(Detector);
+					});
+			});
+		REQUIRE(AutoRTFM::ETransactionResult::Committed == Result);
+	}
+
+	SECTION("Transact(Scope(UE_MT_SCOPED_WRITE_ACCESS), Transact(UE_MT_SCOPED_READ_ACCESS))")
+	{
+		FRWTransactionallySafeAccessDetector Detector;
 		AutoRTFM::ETransactionResult Result = AutoRTFM::Transact([&]
-		{
-			Detector.ReleaseReadAccess();
-			Detector.AcquireWriteAccess();
-			Detector.ReleaseWriteAccess();
-		});
+			{
+				{
+					UE_MT_SCOPED_WRITE_ACCESS(Detector);
+				}
+
+				AutoRTFM::Commit([&]
+					{
+						UE_MT_SCOPED_READ_ACCESS(Detector);
+					});
+			});
 		REQUIRE(AutoRTFM::ETransactionResult::Committed == Result);
 	}
 
-	SECTION("AcquireReadAccess, Transact(ReleaseReadAccess, AcquireWriteAccess, ReleaseWriteAccess, Abort), ReleaseReadAccess")
+	SECTION("Transact(Scope(UE_MT_SCOPED_WRITE_ACCESS), Transact(UE_MT_SCOPED_WRITE_ACCESS))")
 	{
-		Detector.AcquireReadAccess();
+		FRWTransactionallySafeAccessDetector Detector;
+		AutoRTFM::ETransactionResult Result = AutoRTFM::Transact([&]
+			{
+				{
+					UE_MT_SCOPED_WRITE_ACCESS(Detector);
+				}
+
+				AutoRTFM::Commit([&]
+					{
+						UE_MT_SCOPED_WRITE_ACCESS(Detector);
+					});
+			});
+		REQUIRE(AutoRTFM::ETransactionResult::Committed == Result);
+	}
 
+	SECTION("Transact(Scope(UE_MT_SCOPED_READ_ACCESS), Transact(UE_MT_SCOPED_READ_ACCESS), Abort)")
+	{
+		FRWTransactionallySafeAccessDetector Detector;
 		AutoRTFM::ETransactionResult Result = AutoRTFM::Transact([&]
-		{
-			Detector.ReleaseReadAccess();
-			Detector.AcquireWriteAccess();
-			Detector.ReleaseWriteAccess();
-			AutoRTFM::AbortTransaction();
-		});
+			{
+				{
+					UE_MT_SCOPED_READ_ACCESS(Detector);
+				}
+
+				AutoRTFM::Commit([&]
+					{
+						UE_MT_SCOPED_READ_ACCESS(Detector);
+					});
+
+				AutoRTFM::AbortTransaction();
+			});
 		REQUIRE(AutoRTFM::ETransactionResult::AbortedByRequest == Result);
+	}
+
+	SECTION("Transact(Scope(UE_MT_SCOPED_READ_ACCESS), Transact(UE_MT_SCOPED_WRITE_ACCESS), Abort)")
+	{
+		FRWTransactionallySafeAccessDetector Detector;
+		AutoRTFM::ETransactionResult Result = AutoRTFM::Transact([&]
+			{
+				{
+					UE_MT_SCOPED_READ_ACCESS(Detector);
+				}
+
+				AutoRTFM::Commit([&]
+					{
+						UE_MT_SCOPED_WRITE_ACCESS(Detector);
+					});
 
-		Detector.ReleaseReadAccess();
+				AutoRTFM::AbortTransaction();
+			});
+		REQUIRE(AutoRTFM::ETransactionResult::AbortedByRequest == Result);
 	}
 
+	SECTION("Transact(Scope(UE_MT_SCOPED_WRITE_ACCESS), Transact(UE_MT_SCOPED_READ_ACCESS), Abort)")
 	{
-		// These tests use AutoRTFM::Open() blocks to acquire / release locks.
-		// Transaction retries will cause the lock counts to go out of sync, so
-		// we can only test these with retries disabled retries.
-		struct DisableRetriesScope
-		{
-			DisableRetriesScope() : OldState(AutoRTFM::ForTheRuntime::GetRetryTransaction())
+		FRWTransactionallySafeAccessDetector Detector;
+		AutoRTFM::ETransactionResult Result = AutoRTFM::Transact([&]
 			{
-				AutoRTFM::ForTheRuntime::SetRetryTransaction(AutoRTFM::ForTheRuntime::NoRetry);
-			}
-			~DisableRetriesScope()
+				{
+					UE_MT_SCOPED_WRITE_ACCESS(Detector);
+				}
+
+				AutoRTFM::Commit([&]
+					{
+						UE_MT_SCOPED_READ_ACCESS(Detector);
+					});
+
+				AutoRTFM::AbortTransaction();
+			});
+		REQUIRE(AutoRTFM::ETransactionResult::AbortedByRequest == Result);
+	}
+
+	SECTION("Transact(Scope(UE_MT_SCOPED_WRITE_ACCESS), Transact(UE_MT_SCOPED_WRITE_ACCESS), Abort)")
+	{
+		FRWTransactionallySafeAccessDetector Detector;
+		AutoRTFM::ETransactionResult Result = AutoRTFM::Transact([&]
 			{
-				AutoRTFM::ForTheRuntime::SetRetryTransaction(OldState);
-			}
-			AutoRTFM::ForTheRuntime::EAutoRTFMRetryTransactionState OldState;
-		};
-		DisableRetriesScope DisableRetries;
+				{
+					UE_MT_SCOPED_WRITE_ACCESS(Detector);
+				}
+
+				AutoRTFM::Commit([&]
+					{
+						UE_MT_SCOPED_WRITE_ACCESS(Detector);
+					});
 
-		SECTION("Transact(AcquireReadAccess, Open(AcquireReadAccess, ReleaseReadAccess), ReleaseReadAccess)")
-		{
-			AutoRTFM::ETransactionResult Result = AutoRTFM::Transact([&]
+				AutoRTFM::AbortTransaction();
+			});
+		REQUIRE(AutoRTFM::ETransactionResult::AbortedByRequest == Result);
+	}
+
+	SECTION("Transact(Scope(UE_MT_SCOPED_READ_ACCESS), Transact(UE_MT_SCOPED_READ_ACCESS), Abort)")
+	{
+		FRWTransactionallySafeAccessDetector Detector;
+		AutoRTFM::ETransactionResult Result = AutoRTFM::Transact([&]
 			{
-				Detector.AcquireReadAccess();
-				AutoRTFM::Open([&]
 				{
-					Detector.AcquireReadAccess();
-					Detector.ReleaseReadAccess();
-				});
-				Detector.ReleaseReadAccess();
+					UE_MT_SCOPED_READ_ACCESS(Detector);
+				}
+
+				AutoRTFM::ETransactionResult Result2 = AutoRTFM::Transact([&]
+					{
+						UE_MT_SCOPED_READ_ACCESS(Detector);
+						AutoRTFM::AbortTransaction();
+					});
+				REQUIRE(AutoRTFM::ETransactionResult::AbortedByRequest == Result2);
 			});
-			REQUIRE(AutoRTFM::ETransactionResult::Committed == Result);
-		}
+		REQUIRE(AutoRTFM::ETransactionResult::Committed == Result);
+	}
 
-		SECTION("Transact(AcquireReadAccess, Open(AcquireReadAccess, ReleaseReadAccess), ReleaseReadAccess, Abort)")
-		{
-			AutoRTFM::ETransactionResult Result = AutoRTFM::Transact([&]
+	SECTION("Transact(Scope(UE_MT_SCOPED_READ_ACCESS), Transact(UE_MT_SCOPED_WRITE_ACCESS), Abort)")
+	{
+		FRWTransactionallySafeAccessDetector Detector;
+		AutoRTFM::ETransactionResult Result = AutoRTFM::Transact([&]
 			{
-				Detector.AcquireReadAccess();
-				AutoRTFM::Open([&]
 				{
-					Detector.AcquireReadAccess();
-					Detector.ReleaseReadAccess();
-				});
-				Detector.ReleaseReadAccess();
-				AutoRTFM::AbortTransaction();
+					UE_MT_SCOPED_READ_ACCESS(Detector);
+				}
+
+				AutoRTFM::ETransactionResult Result2 = AutoRTFM::Transact([&]
+					{
+						UE_MT_SCOPED_WRITE_ACCESS(Detector);
+						AutoRTFM::AbortTransaction();
+					});
+				REQUIRE(AutoRTFM::ETransactionResult::AbortedByRequest == Result2);
 			});
-			REQUIRE(AutoRTFM::ETransactionResult::AbortedByRequest == Result);
-		}
+		REQUIRE(AutoRTFM::ETransactionResult::Committed == Result);
+	}
 
-		SECTION("Transact(AcquireReadAccess, Open(AcquireReadAccess)), ReleaseReadAccess, ReleaseReadAccess")
-		{
-			AutoRTFM::ETransactionResult Result = AutoRTFM::Transact([&]
+	SECTION("Transact(Scope(UE_MT_SCOPED_WRITE_ACCESS), Transact(UE_MT_SCOPED_READ_ACCESS), Abort)")
+	{
+		FRWTransactionallySafeAccessDetector Detector;
+		AutoRTFM::ETransactionResult Result = AutoRTFM::Transact([&]
 			{
-				Detector.AcquireReadAccess();
-				AutoRTFM::Open([&]
 				{
-					Detector.AcquireReadAccess();
-				});
+					UE_MT_SCOPED_WRITE_ACCESS(Detector);
+				}
+
+				AutoRTFM::ETransactionResult Result2 = AutoRTFM::Transact([&]
+					{
+						UE_MT_SCOPED_READ_ACCESS(Detector);
+						AutoRTFM::AbortTransaction();
+					});
+				REQUIRE(AutoRTFM::ETransactionResult::AbortedByRequest == Result2);
 			});
-			REQUIRE(AutoRTFM::ETransactionResult::Committed == Result);
+		REQUIRE(AutoRTFM::ETransactionResult::Committed == Result);
+	}
 
-			Detector.ReleaseReadAccess();
-			Detector.ReleaseReadAccess();
-		}
+	SECTION("Transact(Scope(UE_MT_SCOPED_WRITE_ACCESS), Transact(UE_MT_SCOPED_WRITE_ACCESS), Abort)")
+	{
+		FRWTransactionallySafeAccessDetector Detector;
+		AutoRTFM::ETransactionResult Result = AutoRTFM::Transact([&]
+			{
+				{
+					UE_MT_SCOPED_WRITE_ACCESS(Detector);
+				}
+
+				AutoRTFM::ETransactionResult Result2 = AutoRTFM::Transact([&]
+					{
+						UE_MT_SCOPED_WRITE_ACCESS(Detector);
+						AutoRTFM::AbortTransaction();
+					});
+				REQUIRE(AutoRTFM::ETransactionResult::AbortedByRequest == Result2);
+			});
+		REQUIRE(AutoRTFM::ETransactionResult::Committed == Result);
+	}
 
-		SECTION("Transact(Open(AcquireReadAccess), ReleaseReadAccess, AcquireWriteAccess, Abort), ReleaseReadAccess")
-		{
-			AutoRTFM::ETransactionResult Result = AutoRTFM::Transact([&]
+	SECTION("Transact(Scope(UE_MT_SCOPED_READ_ACCESS), OnCommit(UE_MT_SCOPED_READ_ACCESS))")
+	{
+		FRWTransactionallySafeAccessDetector Detector;
+		AutoRTFM::ETransactionResult Result = AutoRTFM::Transact([&]
 			{
-				AutoRTFM::Open([&]
 				{
-					Detector.AcquireReadAccess();
-				});
-				Detector.ReleaseReadAccess();
-				Detector.AcquireWriteAccess();
-				AutoRTFM::AbortTransaction();
+					UE_MT_SCOPED_READ_ACCESS(Detector);
+				}
+
+				AutoRTFM::OnCommit([&]
+					{
+						UE_MT_SCOPED_READ_ACCESS(Detector);
+					});
 			});
-			REQUIRE(AutoRTFM::ETransactionResult::AbortedByRequest == Result);
+		REQUIRE(AutoRTFM::ETransactionResult::Committed == Result);
+	}
 
-			Detector.ReleaseReadAccess();
-		}
+	SECTION("Transact(Scope(UE_MT_SCOPED_READ_ACCESS), OnCommit(UE_MT_SCOPED_WRITE_ACCESS))")
+	{
+		FRWTransactionallySafeAccessDetector Detector;
+		AutoRTFM::ETransactionResult Result = AutoRTFM::Transact([&]
+			{
+				{
+					UE_MT_SCOPED_READ_ACCESS(Detector);
+				}
 
-		SECTION("Transact(AcquireReadAccess, Open(AcquireReadAccess), Abort), ReleaseReadAccess")
-		{
-			AutoRTFM::ETransactionResult Result = AutoRTFM::Transact([&]
+				AutoRTFM::OnCommit([&]
+					{
+						UE_MT_SCOPED_WRITE_ACCESS(Detector);
+					});
+			});
+		REQUIRE(AutoRTFM::ETransactionResult::Committed == Result);
+	}
+
+	SECTION("Transact(Scope(UE_MT_SCOPED_WRITE_ACCESS), OnCommit(UE_MT_SCOPED_READ_ACCESS))")
+	{
+		FRWTransactionallySafeAccessDetector Detector;
+		AutoRTFM::ETransactionResult Result = AutoRTFM::Transact([&]
 			{
-				Detector.AcquireReadAccess();
-				AutoRTFM::Open([&]
 				{
-					Detector.AcquireReadAccess();
-				});
-				AutoRTFM::AbortTransaction();
+					UE_MT_SCOPED_WRITE_ACCESS(Detector);
+				}
+
+				AutoRTFM::OnCommit([&]
+					{
+						UE_MT_SCOPED_READ_ACCESS(Detector);
+					});
 			});
-			REQUIRE(AutoRTFM::ETransactionResult::AbortedByRequest == Result);
+		REQUIRE(AutoRTFM::ETransactionResult::Committed == Result);
+	}
 
-			Detector.ReleaseReadAccess();
-		}
+	SECTION("Transact(Scope(UE_MT_SCOPED_WRITE_ACCESS), OnCommit(UE_MT_SCOPED_WRITE_ACCESS))")
+	{
+		FRWTransactionallySafeAccessDetector Detector;
+		AutoRTFM::ETransactionResult Result = AutoRTFM::Transact([&]
+			{
+				{
+					UE_MT_SCOPED_WRITE_ACCESS(Detector);
+				}
 
-		SECTION("AcquireReadAccess, Transact(AcquireReadAccess, Open(ReleaseReadAccess), Abort)")
-		{
-			Detector.AcquireReadAccess();
+				AutoRTFM::OnCommit([&]
+					{
+						UE_MT_SCOPED_WRITE_ACCESS(Detector);
+					});
+			});
+		REQUIRE(AutoRTFM::ETransactionResult::Committed == Result);
+	}
 
-			AutoRTFM::ETransactionResult Result = AutoRTFM::Transact([&]
+	SECTION("Transact(Scope(UE_MT_SCOPED_READ_ACCESS), OnAbort(UE_MT_SCOPED_READ_ACCESS))")
+	{
+		FRWTransactionallySafeAccessDetector Detector;
+		AutoRTFM::ETransactionResult Result = AutoRTFM::Transact([&]
 			{
-				Detector.AcquireReadAccess();
-				AutoRTFM::Open([&]
 				{
-					Detector.ReleaseReadAccess();
-				});
+					UE_MT_SCOPED_READ_ACCESS(Detector);
+				}
+
+				AutoRTFM::OnAbort([&]
+					{
+						UE_MT_SCOPED_READ_ACCESS(Detector);
+					});
+
 				AutoRTFM::AbortTransaction();
 			});
-			REQUIRE(AutoRTFM::ETransactionResult::AbortedByRequest == Result);
-		}
-
-		SECTION("AcquireWriteAccess, Transact(Open(ReleaseWriteAccess), AcquireReadAccess, ReleaseReadAccess)")
-		{
-			Detector.AcquireWriteAccess();
+		REQUIRE(AutoRTFM::ETransactionResult::AbortedByRequest == Result);
+	}
 
-			AutoRTFM::ETransactionResult Result = AutoRTFM::Transact([&]
+	SECTION("Transact(Scope(UE_MT_SCOPED_READ_ACCESS), OnAbort(UE_MT_SCOPED_WRITE_ACCESS))")
+	{
+		FRWTransactionallySafeAccessDetector Detector;
+		AutoRTFM::ETransactionResult Result = AutoRTFM::Transact([&]
 			{
-				AutoRTFM::Open([&]
 				{
-					Detector.ReleaseWriteAccess();
-				});
-				Detector.AcquireReadAccess();
+					UE_MT_SCOPED_READ_ACCESS(Detector);
+				}
+
+				AutoRTFM::OnAbort([&]
+					{
+						UE_MT_SCOPED_WRITE_ACCESS(Detector);
+					});
+
+				AutoRTFM::AbortTransaction();
 			});
-			REQUIRE(AutoRTFM::ETransactionResult::Committed == Result);
+		REQUIRE(AutoRTFM::ETransactionResult::AbortedByRequest == Result);
+	}
 
-			Detector.ReleaseReadAccess();
-		}
+	SECTION("Transact(Scope(UE_MT_SCOPED_WRITE_ACCESS), OnAbort(UE_MT_SCOPED_READ_ACCESS))")
+	{
+		FRWTransactionallySafeAccessDetector Detector;
+		AutoRTFM::ETransactionResult Result = AutoRTFM::Transact([&]
+			{
+				{
+					UE_MT_SCOPED_WRITE_ACCESS(Detector);
+				}
+
+				AutoRTFM::OnAbort([&]
+					{
+						UE_MT_SCOPED_READ_ACCESS(Detector);
+					});
 
-		SECTION("AcquireWriteAccess, Transact(Open(ReleaseWriteAccess), AcquireReadAccess, Abort)")
-		{
-			Detector.AcquireWriteAccess();
+				AutoRTFM::AbortTransaction();
+			});
+		REQUIRE(AutoRTFM::ETransactionResult::AbortedByRequest == Result);
+	}
 
-			AutoRTFM::ETransactionResult Result = AutoRTFM::Transact([&]
+	SECTION("Transact(Scope(UE_MT_SCOPED_WRITE_ACCESS), OnAbort(UE_MT_SCOPED_WRITE_ACCESS), Abort)")
+	{
+		FRWTransactionallySafeAccessDetector Detector;
+		AutoRTFM::ETransactionResult Result = AutoRTFM::Transact([&]
 			{
-				AutoRTFM::Open([&]
 				{
-					Detector.ReleaseWriteAccess();
-				});
-				Detector.AcquireReadAccess();
+					UE_MT_SCOPED_WRITE_ACCESS(Detector);
+				}
+
+				AutoRTFM::OnAbort([&]
+					{
+						UE_MT_SCOPED_WRITE_ACCESS(Detector);
+					});
+
 				AutoRTFM::AbortTransaction();
 			});
-			REQUIRE(AutoRTFM::ETransactionResult::AbortedByRequest == Result);
-		}
+		REQUIRE(AutoRTFM::ETransactionResult::AbortedByRequest == Result);
+	}
 
-		SECTION("AcquireReadAccess, Transact(Open(ReleaseReadAccess), AcquireWriteAccess), ReleaseWriteAccess")
-		{
-			Detector.AcquireReadAccess();
+	SECTION("Transact(Create, UE_MT_SCOPED_READ_ACCESS)")
+	{
+		AutoRTFM::ETransactionResult Result = AutoRTFM::Transact([&]
+			{
+				FRWTransactionallySafeAccessDetector Detector;
+				UE_MT_SCOPED_READ_ACCESS(Detector);
+			});
+		REQUIRE(AutoRTFM::ETransactionResult::Committed == Result);
+	}
 
-			AutoRTFM::ETransactionResult Result = AutoRTFM::Transact([&]
+	SECTION("Transact(Create, UE_MT_SCOPED_READ_ACCESS, Abort)")
+	{
+		AutoRTFM::ETransactionResult Result = AutoRTFM::Transact([&]
 			{
-				AutoRTFM::Open([&]
-				{
-					Detector.ReleaseReadAccess();
-				});
-				Detector.AcquireWriteAccess();
+				FRWTransactionallySafeAccessDetector Detector;
+				UE_MT_SCOPED_READ_ACCESS(Detector);
+				AutoRTFM::AbortTransaction();
 			});
-			REQUIRE(AutoRTFM::ETransactionResult::Committed == Result);
-		
-			Detector.ReleaseWriteAccess();
-		}
+		REQUIRE(AutoRTFM::ETransactionResult::AbortedByRequest == Result);
+	}
 
-		SECTION("AcquireReadAccess, Transact(Open(ReleaseReadAccess), AcquireWriteAccess, Abort)")
-		{
-			Detector.AcquireReadAccess();
+	SECTION("Transact(Create, UE_MT_SCOPED_WRITE_ACCESS)")
+	{
+		AutoRTFM::ETransactionResult Result = AutoRTFM::Transact([&]
+			{
+				FRWTransactionallySafeAccessDetector Detector;
+				UE_MT_SCOPED_WRITE_ACCESS(Detector);
+			});
+		REQUIRE(AutoRTFM::ETransactionResult::Committed == Result);
+	}
 
-			AutoRTFM::ETransactionResult Result = AutoRTFM::Transact([&]
+	SECTION("Transact(Create, UE_MT_SCOPED_WRITE_ACCESS, Abort)")
+	{
+		AutoRTFM::ETransactionResult Result = AutoRTFM::Transact([&]
 			{
-				AutoRTFM::Open([&]
-				{
-					Detector.ReleaseReadAccess();
-				});
-				Detector.AcquireWriteAccess();
+				FRWTransactionallySafeAccessDetector Detector;
+				UE_MT_SCOPED_WRITE_ACCESS(Detector);
 				AutoRTFM::AbortTransaction();
 			});
-			REQUIRE(AutoRTFM::ETransactionResult::AbortedByRequest == Result);
-		}
+		REQUIRE(AutoRTFM::ETransactionResult::AbortedByRequest == Result);
 	}
 }
 
diff --git a/Engine/Source/Runtime/Core/Public/Misc/MTAccessDetector.h b/Engine/Source/Runtime/Core/Public/Misc/MTAccessDetector.h
--- a/Engine/Source/Runtime/Core/Public/Misc/MTAccessDetector.h
+++ b/Engine/Source/Runtime/Core/Public/Misc/MTAccessDetector.h
@@ -107,6 +107,25 @@ public:
 		return !ErrorDetected;
 	}
 
+	/**
+	 * Upgrades a previously taken `AcquireReadAccess` to `AcquireWriteAccess`, and will check if there are other readers or writers
+	 * @return true if no errors were detected
+	 */
+	FORCEINLINE bool UpgradeReadAccessToWriteAccess() const
+	{
+		const bool bErrorDetected = AtomicValue.fetch_add(WriterIncrementValue, std::memory_order_relaxed) != 1;
+
+		if (UNLIKELY(bErrorDetected))
+		{
+			checkf(GIsAutomationTesting, TEXT("Upgrading a write access while there are ongoing read or write access"));
+			return false;
+		}
+
+		// Since we now hold the write lock, we can remove the read lock access we've upgraded to write lock access without re-checking!
+		AtomicValue -= 1;
+		return true;
+	}
+
 protected:
 
 	// We need to do an atomic operation to know there are multiple writers, this is why we reserve more than one bit for them.
diff --git a/Engine/Source/Runtime/Core/Public/Misc/MTTransactionallySafeAccessDetector.h b/Engine/Source/Runtime/Core/Public/Misc/MTTransactionallySafeAccessDetector.h
--- a/Engine/Source/Runtime/Core/Public/Misc/MTTransactionallySafeAccessDetector.h
+++ b/Engine/Source/Runtime/Core/Public/Misc/MTTransactionallySafeAccessDetector.h
@@ -6,63 +6,119 @@
 #include "HAL/Platform.h"
 #include "MTAccessDetector.h"
 #include "Misc/AssertionMacros.h"
+#include "Templates/SharedPointer.h"
 
 #if ENABLE_MT_DETECTOR
 
-/**
- * A version of FRWAccessDetector which also can be used in AutoRTFM transactions.
- *
- * When methods are called outside of a closed transaction they will behave
- * exactly the same as FRWAccessDetector.
- *
- * When called from a closed transaction:
- *  - The first call to AcquireWriteLock() will hold the internal write lock for
- *    the remaining duration of the transaction, even if the transaction calls
- *    ReleaseWriteLock(). This is done to guard against writes that would occur
- *    if the transaction were aborted. The transaction's use of the API remains
- *    unchanged, but unlocks and read locks are effectively no-ops and other 
- *    threads will "see" a longer duration of the write lock.
- *  - Each call to [Acquire|Release][Read|Write]Access is recorded as a pair
- *    of signed read and write counters so that the OnAbort and OnCommit
- *    handlers can restore or apply the read and write lock state, respectively.
- */
+// A transactionally safe access detector that works in the following novel ways:
+// - In the open (non-transactional):
+//   - Acquires the read/write access like before.
+//   - Release the read/write access like before.
+// - In the closed (transactional):
+//   - During acquiring read access we query `ReadLockCount`:
+//     - 0 means we haven't taken read access before in our transaction nest and we acquire it.
+//     - But **only** if we haven't previously taken write access (by querying `WriteLockCount`).
+//     - Then we bump `ReadLockCount` to remember we did a read.
+//     - We also register an on-abort handler to release the access.
+//   - During acquiring write access we query `WriteLockCount`:
+//     - 0 means we haven't taken write access before in our transaction nest and we acquire it.
+//     - But if `ReadLockCount` was non-zero then we have to upgrade the access from read to write.
+//     - Then we bump `WriteLockCount` to remember we did a write.
+//   - During releases we defer these to on-commit.
+// 
+// During on-commit we always release all our `ReadLockCount`'s first, so that we handle the case
+// correctly where we had read then write, we need to only actually release the write access and
+// this means we correctly handle that case.
 struct FRWTransactionallySafeAccessDetectorDefinition
 {
-public:
+	// Always open because the constructor arguments will create the underlying state.
+	UE_AUTORTFM_ALWAYS_OPEN FRWTransactionallySafeAccessDetectorDefinition() : State(MakeShared<FState>())
+	{
+		if (AutoRTFM::IsTransactional())
+		{
+			const AutoRTFM::EContextStatus Status = AutoRTFM::Close([this]
+				{
+					AutoRTFM::PushOnAbortHandler(this, [this]
+						{
+							this->~FRWTransactionallySafeAccessDetectorDefinition();
+						});
+				});
+
+			ensure(AutoRTFM::EContextStatus::OnTrack == Status);
+		}
+	}
+
+	~FRWTransactionallySafeAccessDetectorDefinition()
+	{
+		if (AutoRTFM::IsTransactional())
+		{
+			const AutoRTFM::EContextStatus Status = AutoRTFM::Close([this]
+				{
+					AutoRTFM::PopOnAbortHandler(this);
+
+					// We explicitly copy the state here for the case that `this` was stack
+					// allocated and has already died before the on-commit is hit.
+					AutoRTFM::OnCommit([State = this->State]
+						{
+							ensure((0 == State->ReadLockCount) && (0 == State->WriteLockCount));
+						});
+				});
+
+			ensure(AutoRTFM::EContextStatus::OnTrack == Status);
+		}
+
+		// As the State was constructed in the open, it must be released in the open.
+		AutoRTFM::Open([&] { State = nullptr; });
+	}
+
+	// All the definitions in here are private because we do not want arbitrary code
+	// to call the acquire/release read/write access functions directly, and instead
+	// they should use `UE_MT_SCOPED_READ_ACCESS` and `UE_MT_SCOPED_WRITE_ACCESS`.
+	// This restriction means that all acquires and releases are pairs that must
+	// happen together in the open, or together in closed.
+private:
 	/**
 	 * Acquires read access, will check if there are any writers
 	 * @return true if no errors were detected
 	 */
 	FORCEINLINE bool AcquireReadAccess() const
 	{
-		if (!AutoRTFM::IsClosed())
+		if (AutoRTFM::IsTransactional() || AutoRTFM::IsCommittingOrAborting())
 		{
-			// Called outside of a closed transaction.
-			// Forward to the inner detector and return.
-			return Inner.AcquireReadAccess();
-		}
+			const bool bResult = AutoRTFM::Open([&]
+				{
+					// The transactional system which can increment the lock counts
+					// is always single-threaded, thus this is safe to check without atomicity.
+					if (0 == State->ReadLockCount && 0 == State->WriteLockCount)
+					{
+						if (UNLIKELY(!State->Detector.AcquireReadAccess()))
+						{
+							return false;
+						}
+					}
 
-		// AcquireReadAccess() called inside a closed transaction.
-		return AutoRTFM::Open([this]
-		{
-			if (!TransactionalState.HoldsInternalWriteLock)
+					State->ReadLockCount += 1;
+					return true;
+				});
+
+			if (UNLIKELY(!bResult))
 			{
-				// AcquireReadAccess() called before the first transactional
-				// call to AcquireWriteAccess(). Call the inner detector.
-				if (UNLIKELY(!Inner.AcquireReadAccess()))
-				{
-					return false;
-				}
+				return false;
+			}
 
-				// Register the abort handler if this is the first transactional
-				// method call.
-				MaybeRegisterAbortHandler();
+			RegisterOnAbortRelease();
+			return true;
+		}
+		else
+		{
+			if (UNLIKELY(!State->Detector.AcquireReadAccess()))
+			{
+				return false;
 			}
 
-			// Increment the transactional read-lock counter.
-			TransactionalState.ReadLockDelta++;
+			ensure(0 == State->WriteLockCount);
 			return true;
-		});
+		}
 	}
 
 	/**
@@ -71,34 +127,19 @@ public:
 	 */
 	FORCEINLINE bool ReleaseReadAccess() const
 	{
-		if (!AutoRTFM::IsClosed())
+		if (AutoRTFM::IsTransactional() || AutoRTFM::IsCommittingOrAborting())
 		{
-			// Called outside of a closed transaction.
-			// Forward to the inner detector and return.
-			return Inner.ReleaseReadAccess();
-		}
+			RegisterOnCommitRelease();
 
-		// ReleaseReadAccess() called inside a closed transaction.
-		return AutoRTFM::Open([this]
-		{
-			if (!TransactionalState.HoldsInternalWriteLock)
-			{
-				// ReleaseReadAccess() called before the first transactional
-				// call to AcquireWriteAccess(). Call the inner detector.
-				if (UNLIKELY(!Inner.ReleaseReadAccess()))
-				{
-					return false;
-				}
-
-				// Register the abort handler if this is the first transactional
-				// method call.
-				MaybeRegisterAbortHandler();
-			}
-
-			// Decrement the transactional read-lock counter.
-			TransactionalState.ReadLockDelta--;
+			// We can't do anything better here than returning true, because we
+			// are deferring the actual release until on commit!
 			return true;
-		});
+		}
+		else
+		{
+			ensure(0 == State->WriteLockCount);
+			return State->Detector.ReleaseReadAccess();
+		}
 	}
 
 	/** 
@@ -107,37 +148,51 @@ public:
 	 */
 	FORCEINLINE bool AcquireWriteAccess() const
 	{
-		if (!AutoRTFM::IsClosed())
+		if (AutoRTFM::IsTransactional() || AutoRTFM::IsCommittingOrAborting())
 		{
-			// Called outside of a closed transaction.
-			// Forward to the inner detector and return.
-			return Inner.AcquireWriteAccess();
-		}
+			const bool bResult = AutoRTFM::Open([&]
+				{
+					if ((0 == State->ReadLockCount) && (0 == State->WriteLockCount))
+					{
+						// There have been no prior calls to `AcquireReadAccess`
+						// so we can just claim write access directly.
+						if (UNLIKELY(!State->Detector.AcquireWriteAccess()))
+						{
+							return false;
+						}
+					}
+					else if (0 == State->WriteLockCount)
+					{
+						// There was a prior call to `AcquireReadAccess` so we need
+						// to upgrade our read access to write access.
+						if (UNLIKELY(!State->Detector.UpgradeReadAccessToWriteAccess()))
+						{
+							return false;
+						}
+					}
+
+					State->WriteLockCount += 1;
+					return true;
+				});
+
+			if (UNLIKELY(!bResult))
+			{
+				return false;
+			}
 
-		return AutoRTFM::Open([this]
+			RegisterOnAbortRelease();
+			return true;
+		}
+		else
 		{
-			if (!TransactionalState.HoldsInternalWriteLock)
+			if (UNLIKELY(!State->Detector.AcquireWriteAccess()))
 			{
-				// First transactional call to AcquireWriteAccess(). 
-				// Call the inner detector.
-				if (UNLIKELY(!Inner.AcquireWriteAccess()))
-				{
-					return false;
-				}
-
-				// Register the abort handler if this is the first transactional
-				// method call.
-				MaybeRegisterAbortHandler();
-				// Register the commit handler to unlock the write-lock which
-				// will be held until the transaction is done.
-				RegisterCommitHandler();
-				TransactionalState.HoldsInternalWriteLock = true;
+				return false;
 			}
 
-			// Increment the transactional write-lock counter.
-			TransactionalState.WriteLockDelta++;
+			ensure((0 == State->ReadLockCount) && (0 == State->WriteLockCount));
 			return true;
-		});
+		}
 	}
 
 	/** 
@@ -146,148 +201,98 @@ public:
 	 */
 	FORCEINLINE bool ReleaseWriteAccess() const
 	{
-		if (!AutoRTFM::IsClosed())
+		if (AutoRTFM::IsTransactional() || AutoRTFM::IsCommittingOrAborting())
 		{
-			// Called outside of a closed transaction.
-			// Forward to the inner detector and return.
-			return Inner.ReleaseWriteAccess();
-		}
+			RegisterOnCommitRelease();
 
-		return AutoRTFM::Open([this]
-		{
-			if (!TransactionalState.HoldsInternalWriteLock)
-			{
-				// ReleaseWriteAccess() called before the first transactional
-				// call to AcquireWriteAccess(). Call the inner detector.
-				if (UNLIKELY(!Inner.ReleaseWriteAccess()))
-				{
-					return false;
-				}
-
-				// Register the abort handler if this is the first transactional
-				// method call.
-				MaybeRegisterAbortHandler();
-			}
-
-			// Decrement the transactional write-lock counter.
-			TransactionalState.WriteLockDelta--;
+			// We can't do anything better here than returning true, because we
+			// are deferring the actual release until on commit!
 			return true;
-		});
+		}
+		else
+		{
+			ensure((0 == State->ReadLockCount) && (0 == State->WriteLockCount));
+			return State->Detector.ReleaseWriteAccess();
+		}
 	}
 
 private:
-	// Registers a transaction abort handler if this is the first call for the
-	// current transaction.
-	void MaybeRegisterAbortHandler() const
+	struct FState final
 	{
-		if (TransactionalState.AbortHandlerRegistered)
-		{
-			return; // Already registered.
-		}
-		TransactionalState.AbortHandlerRegistered = true;
+		// The underlying FRWAccessDetector.
+		FRWAccessDetector Detector;
 
-		// Call OnAbort() in the closed, otherwise it'll be a no-op.
-		(void) AutoRTFM::Close([&]
-		{
-			AutoRTFM::OnAbort([this]
-			{
-				// Transaction is being aborted. Undo all state changes.
-				if (TransactionalState.HoldsInternalWriteLock)
-				{
-					// AcquireWriteAccess() was called in the transaction.
-					// The initial state when entering the transaction must have
-					// been write-unlocked, otherwise we'd double-write-lock.
-					// Unlock the write.
-					Inner.ReleaseWriteAccess();
-				}
-				else
-				{
-					// AcquireWriteAccess() was not called in the transaction.
-					// Undo any transactional call ReleaseWriteAccess().
-					if (TransactionalState.WriteLockDelta < 0)
-					{
-						Inner.AcquireWriteAccess();
-					}
-					// Undo any transactional calls AcquireReadAccess().
-					for (; TransactionalState.ReadLockDelta > 0; TransactionalState.ReadLockDelta--)
-					{
-						Inner.ReleaseReadAccess();
-					}
-				}
-				
-				// Undo any transactional calls ReleaseReadAccess().
-				for (; TransactionalState.ReadLockDelta < 0; TransactionalState.ReadLockDelta++)
-				{
-					Inner.AcquireReadAccess();
-				}
+		// True if the abort handler has been registered.
+		bool AbortHandlerRegistered = false;
 
-				// Clear the TransactionalState for the next transaction.
-				TransactionalState = {};
-			});
-		});
-	}
+		unsigned ReadLockCount = 0;
+
+		unsigned WriteLockCount = 0;
+	};
 
-	// Registers a transaction commit handler to rebalance locks after the first
-	// call to AcquireWriteAccess().
-	void RegisterCommitHandler() const
+	UE_AUTORTFM_NOAUTORTFM static void ReleaseAccess(TSharedPtr<FState> State)
 	{
-		// Call OnCommit() in the closed, otherwise it'll invoke the callback 
-		// immediately instead of when the transaction is committed. 
-		(void) AutoRTFM::Close([&]
+		if (0 < State->ReadLockCount)
 		{
-			AutoRTFM::OnCommit([this]
+			State->ReadLockCount -= 1;
+
+			if ((0 == State->ReadLockCount) && (0 == State->WriteLockCount))
 			{
-				// Transaction is being committed.
-				check(TransactionalState.HoldsInternalWriteLock);
-				check(TransactionalState.WriteLockDelta >= 0);
-				// Release the write-lock, if the write-lock is now balanced.
-				if (TransactionalState.WriteLockDelta == 0)
-				{
-					Inner.ReleaseWriteAccess();
-				}
-				// Calls to AcquireReadAccess() with the write-lock held 
-				// were a no-op, so apply them now.
-				for (; TransactionalState.ReadLockDelta > 0; TransactionalState.ReadLockDelta--)
-				{
-					Inner.AcquireReadAccess();
-				}
+				State->Detector.ReleaseReadAccess();
+			}
+		}
+		else if (0 < State->WriteLockCount)
+		{
+			State->WriteLockCount -= 1;
 
-				// Clear the TransactionalState for the next transaction.
-				TransactionalState = {};
-			});
-		});
+			if (0 == State->WriteLockCount)
+			{
+				State->Detector.ReleaseWriteAccess();
+			}
+		}
+		else
+		{
+			// We should only register as many on-abort handlers as we had lock count increments!
+			check(false);
+		}
 	}
 
-	struct FTransactionalState
+	FORCEINLINE void RegisterOnAbortRelease() const
 	{
-		// Incremented for each transactional call to AcquireReadAccess(),
-		// decremented for each transactional call to ReleaseReadAccess().
-		int16 ReadLockDelta = 0;
-
-		// Incremented for each transactional call to AcquireWriteAccess(),
-		// decremented for each transactional call to ReleaseWriteAccess().
-		int16 WriteLockDelta = 0;
-		
-		// If true, then the write-lock has been taken in a transaction and will
-		// be held until the transaction is complete (either aborted or 
-		// committed).
-		bool HoldsInternalWriteLock = false;
-
-		// True if the abort handler has been registered.
-		bool AbortHandlerRegistered = false;
-	};
+		// We explicitly copy the state here for the case that `this` was stack
+		// allocated and has already died before the on-abort is hit.
+		AutoRTFM::OnAbort([State = this->State] { ReleaseAccess(State); });
+	}
 
-	// The inner FRWAccessDetector.
-	FRWAccessDetector Inner;
+	FORCEINLINE void RegisterOnCommitRelease() const
+	{
+		// We explicitly copy the state here for the case that `this` was stack
+		// allocated and has already died before the on-abort is hit.
+		AutoRTFM::OnCommit([State = this->State] { ReleaseAccess(State); });
+	}
 
 	// The state held for calls made when in a transaction.
-	mutable FTransactionalState TransactionalState;
+	TSharedPtr<FState> State;
+
+	friend TScopedReaderAccessDetector<FRWTransactionallySafeAccessDetectorDefinition>;
+	friend TScopedWriterDetector<FRWTransactionallySafeAccessDetectorDefinition>;
+	friend TScopedReaderAccessDetector<const FRWTransactionallySafeAccessDetectorDefinition>;
+	friend TScopedWriterDetector<const FRWTransactionallySafeAccessDetectorDefinition>;
+};
+
+// TODO: if we made `FRWAccessDetector` have private + friend like below we don't need this anymore!
+struct FRWFallbackSafeAccessDetectorDefinition : private FRWAccessDetector
+{
+	friend TScopedReaderAccessDetector<FRWFallbackSafeAccessDetectorDefinition>;
+	friend TScopedWriterDetector<FRWFallbackSafeAccessDetectorDefinition>;
+	friend TScopedReaderAccessDetector<const FRWFallbackSafeAccessDetectorDefinition>;
+	friend TScopedWriterDetector<const FRWFallbackSafeAccessDetectorDefinition>;
 };
 
 #if UE_AUTORTFM
 using FRWTransactionallySafeAccessDetector = FRWTransactionallySafeAccessDetectorDefinition;
 #else
-using FRWTransactionallySafeAccessDetector = FRWAccessDetector;
+using FRWTransactionallySafeAccessDetector = FRWFallbackSafeAccessDetectorDefinition;
 #endif
 
 #define UE_MT_DECLARE_TS_RW_ACCESS_DETECTOR(AccessDetector) FRWTransactionallySafeAccessDetector AccessDetector;
