[Cleanup] Make indentation in Abort tests consistent, and remove UintToPointer.

The UintToPointer helper was not necessary, and could be replaced by a real pointer to an arbitrary value.

diff --git a/Engine/Source/Programs/AutoRTFMTests/Private/Abort.cpp b/Engine/Source/Programs/AutoRTFMTests/Private/Abort.cpp
--- a/Engine/Source/Programs/AutoRTFMTests/Private/Abort.cpp
+++ b/Engine/Source/Programs/AutoRTFMTests/Private/Abort.cpp
@@ -57,26 +57,26 @@ TEST_CASE("Abort.NestedAbortOrder")
 	{
 		// If we are retrying transactions, need to reset the test state.
 		AutoRTFM::OnAbort([&]
-			{
-				Orderer = 0;
-			});
+		{
+			Orderer = 0;
+		});
 
 		InnerResult = AutoRTFM::Transact([&]
-			{
-				AutoRTFM::OnAbort([&]
-					{
-						REQUIRE(1 == Orderer);
-						Orderer += 1;
-					});
+		{
+			AutoRTFM::OnAbort([&]
+				{
+					REQUIRE(1 == Orderer);
+					Orderer += 1;
+				});
 
-				AutoRTFM::OnAbort([&]
-					{
-						REQUIRE(0 == Orderer);
-						Orderer += 1;
-					});
+			AutoRTFM::OnAbort([&]
+				{
+					REQUIRE(0 == Orderer);
+					Orderer += 1;
+				});
 
-				AutoRTFM::AbortTransaction();
-			});
+			AutoRTFM::AbortTransaction();
+		});
 	});
 
 	REQUIRE(AutoRTFM::ETransactionResult::AbortedByRequest == InnerResult);
@@ -251,17 +251,7 @@ TEST_CASE("Abort.CascadeThroughManualTransaction")
 	REQUIRE(false == bTouched);
 }
 
-inline void* UIntToPointer(uint64 Value)
-{
-	union
-	{
-		uint64 Int;
-		void* Ptr;
-	};
-
-	Int = Value;
-	return Ptr;
-}
+static int SAbortKey;
 
 TEST_CASE("Abort.PushOnAbortHandler_NoAbort")
 {
@@ -270,7 +260,7 @@ TEST_CASE("Abort.PushOnAbortHandler_NoAbort")
 	AutoRTFM::Commit([&]
 	{
 		Value = 66;
-		AutoRTFM::PushOnAbortHandler(UIntToPointer(747), [&Value](){ Value = 77; });
+		AutoRTFM::PushOnAbortHandler(&SAbortKey, [&Value](){ Value = 77; });
 	});
 
 	REQUIRE(Value == 66);
@@ -283,7 +273,7 @@ TEST_CASE("Abort.PushOnAbortHandler_WithAbort")
 	const AutoRTFM::ETransactionResult Result = AutoRTFM::Transact([&]
 	{
 		Value = 66;
-		AutoRTFM::PushOnAbortHandler(UIntToPointer(747), [&Value](){ Value = 77; });
+		AutoRTFM::PushOnAbortHandler(&SAbortKey, [&Value](){ Value = 77; });
 
 		AutoRTFM::AbortTransaction();
 	});
@@ -298,7 +288,7 @@ TEST_CASE("Abort.PushOnAbortHandler_MutableCapture")
 
 	AutoRTFM::Transact([&]
 	{
-		AutoRTFM::PushOnAbortHandler(UIntToPointer(12345), [MessageCopy = Message]() mutable
+		AutoRTFM::PushOnAbortHandler(&SAbortKey, [MessageCopy = Message]() mutable
 		{
 			MessageCopy += " World!";
 			REQUIRE(MessageCopy == "Hello World!");
@@ -313,13 +303,13 @@ TEST_CASE("Abort.PushOnAbortHandler_WithPop_NoAbort")
 	int Value = 55;
 
 	AutoRTFM::Commit([&]
-		{
-			Value = 66;
-			AutoRTFM::PushOnAbortHandler(UIntToPointer(747), [&Value]() { Value = 77; });
-			Value = 88;
+	{
+		Value = 66;
+		AutoRTFM::PushOnAbortHandler(&SAbortKey, [&Value]() { Value = 77; });
+		Value = 88;
 
-			AutoRTFM::PopOnAbortHandler(UIntToPointer(747));
-		});
+		AutoRTFM::PopOnAbortHandler(&SAbortKey);
+	});
 
 	REQUIRE(Value == 88);
 }
@@ -331,10 +321,10 @@ TEST_CASE("Abort.PushOnAbortHandler_WithPopAll_NoAbort")
 	AutoRTFM::Commit([&]
 	{
 		Value = 66;
-		AutoRTFM::PushOnAbortHandler(UIntToPointer(747), [&Value](){ Value = 77; });
+		AutoRTFM::PushOnAbortHandler(&SAbortKey, [&Value](){ Value = 77; });
 		Value = 88;
 
-		AutoRTFM::PopAllOnAbortHandlers(UIntToPointer(747));
+		AutoRTFM::PopAllOnAbortHandlers(&SAbortKey);
 	});
 
 	REQUIRE(Value == 88);
@@ -347,10 +337,10 @@ TEST_CASE("Abort.PushOnAbortHandler_WithPop_WithAbort")
 	const AutoRTFM::ETransactionResult Result = AutoRTFM::Transact([&]
 	{
 		Value = 66;
-		AutoRTFM::PushOnAbortHandler(UIntToPointer(747), [&Value](){ Value = 77; });
+		AutoRTFM::PushOnAbortHandler(&SAbortKey, [&Value](){ Value = 77; });
 		Value = 88;
 
-		AutoRTFM::PopOnAbortHandler(UIntToPointer(747));
+		AutoRTFM::PopOnAbortHandler(&SAbortKey);
 
 		AutoRTFM::AbortTransaction();
 	});
@@ -366,10 +356,10 @@ TEST_CASE("Abort.PushOnAbortHandler_WithPopAll_WithAbort")
 	const AutoRTFM::ETransactionResult Result = AutoRTFM::Transact([&]
 	{
 		Value = 66;
-		AutoRTFM::PushOnAbortHandler(UIntToPointer(747), [&Value]() { Value = 77; });
+		AutoRTFM::PushOnAbortHandler(&SAbortKey, [&Value]() { Value = 77; });
 		Value = 88;
 
-		AutoRTFM::PopAllOnAbortHandlers(UIntToPointer(747));
+		AutoRTFM::PopAllOnAbortHandlers(&SAbortKey);
 
 		AutoRTFM::AbortTransaction();
 	});
@@ -383,16 +373,16 @@ TEST_CASE("Abort.PushOnAbortHandler_Duplicates1")
 	int Value = 55;
 
 	const AutoRTFM::ETransactionResult Result = AutoRTFM::Transact([&]
-		{
-			Value = 66;
-			AutoRTFM::PushOnAbortHandler(UIntToPointer(747), [&Value]() { Value = 77; });
-			AutoRTFM::PushOnAbortHandler(UIntToPointer(747), [&Value]() { Value = 88; });
-			Value = 99;
+	{
+		Value = 66;
+		AutoRTFM::PushOnAbortHandler(&SAbortKey, [&Value]() { Value = 77; });
+		AutoRTFM::PushOnAbortHandler(&SAbortKey, [&Value]() { Value = 88; });
+		Value = 99;
 
-			AutoRTFM::PopOnAbortHandler(UIntToPointer(747));
+		AutoRTFM::PopOnAbortHandler(&SAbortKey);
 
-			AutoRTFM::AbortTransaction();
-		});
+		AutoRTFM::AbortTransaction();
+	});
 
 	REQUIRE(AutoRTFM::ETransactionResult::AbortedByRequest == Result);
 
@@ -407,11 +397,11 @@ TEST_CASE("Abort.PushOnAbortHandler_PopAll_Duplicates")
 	const AutoRTFM::ETransactionResult Result = AutoRTFM::Transact([&]
 	{
 		Value = 66;
-		AutoRTFM::PushOnAbortHandler(UIntToPointer(747), [&Value]() { Value = 77; });
-		AutoRTFM::PushOnAbortHandler(UIntToPointer(747), [&Value]() { Value = 88; });
+		AutoRTFM::PushOnAbortHandler(&SAbortKey, [&Value]() { Value = 77; });
+		AutoRTFM::PushOnAbortHandler(&SAbortKey, [&Value]() { Value = 88; });
 		Value = 99;
 
-		AutoRTFM::PopAllOnAbortHandlers(UIntToPointer(747));
+		AutoRTFM::PopAllOnAbortHandlers(&SAbortKey);
 
 		AutoRTFM::AbortTransaction();
 	});
@@ -429,8 +419,8 @@ TEST_CASE("Abort.PushOnAbortHandler_Duplicates2")
 	const AutoRTFM::ETransactionResult Result = AutoRTFM::Transact([&]
 	{
 		Value = 66;
-		AutoRTFM::PushOnAbortHandler(UIntToPointer(747), [&Value](){ Value += 12; });
-		AutoRTFM::PushOnAbortHandler(UIntToPointer(747), [&Value](){ Value = 65; });
+		AutoRTFM::PushOnAbortHandler(&SAbortKey, [&Value](){ Value += 12; });
+		AutoRTFM::PushOnAbortHandler(&SAbortKey, [&Value](){ Value = 65; });
 		Value = 99;
 
 		AutoRTFM::AbortTransaction();
@@ -449,15 +439,15 @@ TEST_CASE("Abort.PushOnAbortHandler_Order")
 			int Value = 37;
 
 			const AutoRTFM::ETransactionResult Result = AutoRTFM::Transact([&]
-				{
-					AutoRTFM::OnAbort([&Value] { REQUIRE(42 == Value); Value += 1; });
-					AutoRTFM::PushOnAbortHandler(UIntToPointer(747), [&Value]() { REQUIRE(40 == Value); Value += 2; });
-					AutoRTFM::OnAbort([&Value] { REQUIRE(37 == Value); Value += 3; });
+			{
+				AutoRTFM::OnAbort([&Value] { REQUIRE(42 == Value); Value += 1; });
+				AutoRTFM::PushOnAbortHandler(&SAbortKey, [&Value]() { REQUIRE(40 == Value); Value += 2; });
+				AutoRTFM::OnAbort([&Value] { REQUIRE(37 == Value); Value += 3; });
 
-					Value = 99;
+				Value = 99;
 
-					AutoRTFM::AbortTransaction();
-				});
+				AutoRTFM::AbortTransaction();
+			});
 
 			REQUIRE(AutoRTFM::ETransactionResult::AbortedByRequest == Result);
 			REQUIRE(Value == 43);
@@ -468,17 +458,17 @@ TEST_CASE("Abort.PushOnAbortHandler_Order")
 			int Value = 37;
 
 			const AutoRTFM::ETransactionResult Result = AutoRTFM::Transact([&]
-				{
-					AutoRTFM::OnAbort([&Value] { REQUIRE(40 == Value); Value += 1; });
-					AutoRTFM::PushOnAbortHandler(UIntToPointer(747), [&Value]() { REQUIRE(false); });
-					AutoRTFM::OnAbort([&Value] { REQUIRE(37 == Value); Value += 3; });
+			{
+				AutoRTFM::OnAbort([&Value] { REQUIRE(40 == Value); Value += 1; });
+				AutoRTFM::PushOnAbortHandler(&SAbortKey, [&Value]() { REQUIRE(false); });
+				AutoRTFM::OnAbort([&Value] { REQUIRE(37 == Value); Value += 3; });
 
-					AutoRTFM::PopOnAbortHandler(UIntToPointer(747));
+				AutoRTFM::PopOnAbortHandler(&SAbortKey);
 
-					Value = 99;
+				Value = 99;
 
-					AutoRTFM::AbortTransaction();
-				});
+				AutoRTFM::AbortTransaction();
+			});
 
 			REQUIRE(AutoRTFM::ETransactionResult::AbortedByRequest == Result);
 			REQUIRE(Value == 41);
@@ -492,53 +482,53 @@ TEST_CASE("Abort.PushOnAbortHandler_Order")
 			int Value = 37;
 
 			const AutoRTFM::ETransactionResult Result = AutoRTFM::Transact([&]
+			{
+				AutoRTFM::OnAbort([&Value]
 				{
-					AutoRTFM::OnAbort([&Value]
-						{
-							REQUIRE(42 == Value);
-							Value += 1;
-						});
+					REQUIRE(42 == Value);
+					Value += 1;
+				});
 
-					// Make a child transaction.
-					AutoRTFM::Commit([&]
+				// Make a child transaction.
+				AutoRTFM::Commit([&]
+				{
+					AutoRTFM::PushOnAbortHandler(&SAbortKey, [&Value]()
+					{
+						// If we are retrying nested transactions too, we can't check that
+						// the value was something specific before hand!
+						if (!AutoRTFM::ForTheRuntime::ShouldRetryNestedTransactionsToo())
 						{
-							AutoRTFM::PushOnAbortHandler(UIntToPointer(747), [&Value]()
-								{
-									// If we are retrying nested transactions too, we can't check that
-									// the value was something specific before hand!
-									if (!AutoRTFM::ForTheRuntime::ShouldRetryNestedTransactionsToo())
-									{
-										REQUIRE(40 == Value);
-										Value += 2;
-									}
-									else
-									{
-										Value += 1;
-									}
-								});
-						});
-
-					AutoRTFM::OnAbort([&Value]
+							REQUIRE(40 == Value);
+							Value += 2;
+						}
+						else
 						{
-							// If we are retrying nested transactions too, we've ran the on-abort in the
-							// child transaction once, so our value will be larger.
-							if (!AutoRTFM::ForTheRuntime::ShouldRetryNestedTransactionsToo())
-							{
-								REQUIRE(37 == Value);
-							}
-							else
-							{
-								REQUIRE(38 == Value);
-							}
-
-							Value += 3;
-						});
-
-					Value = 99;
+							Value += 1;
+						}
+					});
+				});
 
-					AutoRTFM::AbortTransaction();
+				AutoRTFM::OnAbort([&Value]
+				{
+					// If we are retrying nested transactions too, we've ran the on-abort in the
+					// child transaction once, so our value will be larger.
+					if (!AutoRTFM::ForTheRuntime::ShouldRetryNestedTransactionsToo())
+					{
+						REQUIRE(37 == Value);
+					}
+					else
+					{
+						REQUIRE(38 == Value);
+					}
+
+					Value += 3;
 				});
 
+				Value = 99;
+
+				AutoRTFM::AbortTransaction();
+			});
+
 			REQUIRE(AutoRTFM::ETransactionResult::AbortedByRequest == Result);
 			REQUIRE(Value == 43);
 		}
@@ -548,28 +538,28 @@ TEST_CASE("Abort.PushOnAbortHandler_Order")
 			int Value = 37;
 
 			const AutoRTFM::ETransactionResult Result = AutoRTFM::Transact([&]
-				{
-					AutoRTFM::OnAbort([&Value] { REQUIRE(40 == Value); Value += 1; });
+			{
+				AutoRTFM::OnAbort([&Value] { REQUIRE(40 == Value); Value += 1; });
 
-					// Make a child transaction.
-					AutoRTFM::Commit([&]
-						{
-							AutoRTFM::PushOnAbortHandler(UIntToPointer(747), [&Value]()
-								{
-									// Only if we are retrying on 
-									REQUIRE(AutoRTFM::ForTheRuntime::ShouldRetryNestedTransactionsToo());
-								});
-						});
+				// Make a child transaction.
+				AutoRTFM::Commit([&]
+				{
+					AutoRTFM::PushOnAbortHandler(&SAbortKey, [&Value]()
+					{
+						// Only if we are retrying on 
+						REQUIRE(AutoRTFM::ForTheRuntime::ShouldRetryNestedTransactionsToo());
+					});
+				});
 
-					AutoRTFM::OnAbort([&Value] { REQUIRE(37 == Value); Value += 3; });
+				AutoRTFM::OnAbort([&Value] { REQUIRE(37 == Value); Value += 3; });
 
-					// Bit funky, but we can pop the child's push here!
-					AutoRTFM::PopOnAbortHandler(UIntToPointer(747));
+				// Bit funky, but we can pop the child's push here!
+				AutoRTFM::PopOnAbortHandler(&SAbortKey);
 
-					Value = 99;
+				Value = 99;
 
-					AutoRTFM::AbortTransaction();
-				});
+				AutoRTFM::AbortTransaction();
+			});
 
 			REQUIRE(AutoRTFM::ETransactionResult::AbortedByRequest == Result);
 			REQUIRE(Value == 41);
@@ -581,27 +571,27 @@ TEST_CASE("Abort.PushOnAbortHandler_Order")
 
 			AutoRTFM::ETransactionResult Result = AutoRTFM::ETransactionResult::Committed;
 			AutoRTFM::Commit([&]
-				{
-					AutoRTFM::OnCommit([&Value] { REQUIRE(37 == Value); Value += 1; });
+			{
+				AutoRTFM::OnCommit([&Value] { REQUIRE(37 == Value); Value += 1; });
 
-					// Make a child transaction.
-					Result = AutoRTFM::Transact([&]
-						{
-							AutoRTFM::PushOnAbortHandler(UIntToPointer(747), [&Value]() { REQUIRE(99 == Value); Value += 2; });
-							AutoRTFM::AbortTransaction();
-						});
+				// Make a child transaction.
+				Result = AutoRTFM::Transact([&]
+				{
+					AutoRTFM::PushOnAbortHandler(&SAbortKey, [&Value]() { REQUIRE(99 == Value); Value += 2; });
+					AutoRTFM::AbortTransaction();
+				});
 
-					AutoRTFM::Open([&]
-					{
-						REQUIRE(Value == 101);
-					});
+				AutoRTFM::Open([&]
+				{
+					REQUIRE(Value == 101);
+				});
 
-					AutoRTFM::OnCommit([&Value] { REQUIRE(38 == Value); Value += 3; });
+				AutoRTFM::OnCommit([&Value] { REQUIRE(38 == Value); Value += 3; });
 
-					Value = 37;
+				Value = 37;
 
-					AutoRTFM::OnAbort([&Value] { Value = 99; });
-				});
+				AutoRTFM::OnAbort([&Value] { Value = 99; });
+			});
 
 			REQUIRE(AutoRTFM::ETransactionResult::AbortedByRequest == Result);
 			REQUIRE(Value == 41);
@@ -654,10 +644,10 @@ TEST_CASE("Abort.Language")
 	bool bTouched = false;
 
 	const AutoRTFM::ETransactionResult Result = AutoRTFM::Transact([&]
-		{
-			bTouched = true;
-			FnHasNoClosed();
-		});
+	{
+		bTouched = true;
+		FnHasNoClosed();
+	});
 
 	REQUIRE(AutoRTFM::ETransactionResult::AbortedByLanguage == Result);
 	REQUIRE(false == bTouched);
@@ -668,19 +658,19 @@ TEST_CASE("Abort.LanguageThroughOpen")
 	bool bTouched = false;
 
 	const AutoRTFM::ETransactionResult Result = AutoRTFM::Transact([&]
-		{
-			bTouched = true;
+	{
+		bTouched = true;
 
-			AutoRTFM::Open([&]
-				{
-					const AutoRTFM::EContextStatus Status = AutoRTFM::Close([&]
-						{
-							FnHasNoClosed();
-						});
+		AutoRTFM::Open([&]
+		{
+			const AutoRTFM::EContextStatus Status = AutoRTFM::Close([&]
+			{
+				FnHasNoClosed();
+			});
 
-					REQUIRE(AutoRTFM::EContextStatus::AbortedByLanguage == Status);
-				});
+			REQUIRE(AutoRTFM::EContextStatus::AbortedByLanguage == Status);
 		});
+	});
 
 	REQUIRE(AutoRTFM::ETransactionResult::AbortedByLanguage == Result);
 	REQUIRE(false == bTouched);
