[AutoRTFM] Allow fflush to pass through silently inside a transaction.

This is never actively harmful, and can be harmless. For instance, FFeedbackContextAnsi::LocalPrint calls `fflush(stdout)`. We already expose `printf` as a passthrough function as well.

diff --git a/Engine/Source/Programs/AutoRTFMTests/Private/Tests.cpp b/Engine/Source/Programs/AutoRTFMTests/Private/Tests.cpp
--- a/Engine/Source/Programs/AutoRTFMTests/Private/Tests.cpp
+++ b/Engine/Source/Programs/AutoRTFMTests/Private/Tests.cpp
@@ -7,6 +7,8 @@
 #include <vector>
 #include <cstring>
 #include <cmath>
+#include <cstdio>
+#include <stdio.h>
 #include "Logging/LogMacros.h"
 
 #if PLATFORM_CPU_X86_FAMILY
@@ -343,6 +345,18 @@ TEST_CASE("Tests.RetryNonNested")
 	REQUIRE(2 == Count);
 }
 
+TEST_CASE("Tests.fflush")
+{
+	auto Result = AutoRTFM::Transact([&]()
+	{
+		// There isn't a simple way to verify that fflush has actually done anything,
+		// but we want to at least verify that it can be called safely.
+		fflush(stdout);
+		std::fflush(stdout);
+	});
+	REQUIRE(Result == AutoRTFM::ETransactionResult::Committed);
+}
+
 #if PLATFORM_CPU_X86_FAMILY
 template<int WhichMaskOff>
 void AVXDoMaskedStore(double* Vector)
diff --git a/Engine/Source/Runtime/Core/Private/AutoRTFM/StdLibSupport.cpp b/Engine/Source/Runtime/Core/Private/AutoRTFM/StdLibSupport.cpp
--- a/Engine/Source/Runtime/Core/Private/AutoRTFM/StdLibSupport.cpp
+++ b/Engine/Source/Runtime/Core/Private/AutoRTFM/StdLibSupport.cpp
@@ -518,17 +518,10 @@ int RTFM_wprintf(const wchar_t* Format, ...)
 }
 UE_AUTORTFM_REGISTER_OPEN_FUNCTION(wprintf);
 
-int RTFM_putchar(int Char)
-{
-    return putchar(Char);
-}
-UE_AUTORTFM_REGISTER_OPEN_FUNCTION(putchar);
+UE_AUTORTFM_REGISTER_SELF_FUNCTION(putchar);
+UE_AUTORTFM_REGISTER_SELF_FUNCTION(puts);
 
-int RTFM_puts(const char* Str)
-{
-    return puts(Str);
-}
-UE_AUTORTFM_REGISTER_OPEN_FUNCTION(puts);
+UE_AUTORTFM_REGISTER_SELF_FUNCTION(fflush);
 
 #if PLATFORM_WINDOWS
 
