Rethink the transactionally safe counter (and remove it) - as it didn't provide a strong enough safety net for transactional side effects.

diff --git a/Engine/Source/Programs/AutoRTFMEngineTests/Private/AutoRTFMEngineTests.cpp b/Engine/Source/Programs/AutoRTFMEngineTests/Private/AutoRTFMEngineTests.cpp
--- a/Engine/Source/Programs/AutoRTFMEngineTests/Private/AutoRTFMEngineTests.cpp
+++ b/Engine/Source/Programs/AutoRTFMEngineTests/Private/AutoRTFMEngineTests.cpp
@@ -176,6 +176,7 @@ static void PreventLoadingOfEditorOnlyData()
 			FCoreRedirectObjectName(NAME_None, NAME_None, PackageName),
 			InvalidName);
 	}
+	FCoreRedirects::Initialize();
 	FCoreRedirects::AddRedirectList(NewRedirects, TEXT("AutoRTFMEngineTests.PreventLoadingOfEditorOnlyData"));
 	FCoreRedirects::AddKnownMissing(ECoreRedirectFlags::Type_Package, InvalidName);
 }
diff --git a/Engine/Source/Programs/AutoRTFMTests/Private/LockFreeFixedSizeAllocatorTests.cpp b/Engine/Source/Programs/AutoRTFMTests/Private/LockFreeFixedSizeAllocatorTests.cpp
new file mode 100644
--- /dev/null
+++ b/Engine/Source/Programs/AutoRTFMTests/Private/LockFreeFixedSizeAllocatorTests.cpp
@@ -0,0 +1,253 @@
+// Copyright Epic Games, Inc. All Rights Reserved.
+
+#include "Catch2Includes.h"
+#include "AutoRTFM/AutoRTFMTesting.h"
+#include "Containers/LockFreeFixedSizeAllocator.h"
+#include "HAL/ThreadSafeCounter.h"
+#include "HAL/ThreadSafeCounter64.h"
+
+TEMPLATE_TEST_CASE("LockFreeFixedSizeAllocator", "", FThreadSafeCounter, FThreadSafeCounter64)
+{
+	using ALockFreeFixedSizeAllocator = TLockFreeFixedSizeAllocator<128, PLATFORM_CACHE_LINE_SIZE, TestType>;
+	using IntegerType = TestType::IntegerType;
+
+	SECTION("Transact(Allocate)")
+	{
+		ALockFreeFixedSizeAllocator Allocator;
+		void* Blob = nullptr;
+
+		AutoRTFM::Testing::Commit([&]
+		{
+			Blob = Allocator.Allocate();
+		});
+
+		REQUIRE(nullptr != Blob);
+
+		Allocator.Free(Blob);
+	}
+
+	SECTION("Transact(Allocate, Abort)")
+	{
+		ALockFreeFixedSizeAllocator Allocator;
+		void* Blob = nullptr;
+
+		AutoRTFM::Testing::Abort([&]
+		{
+			Blob = Allocator.Allocate();
+			AutoRTFM::AbortTransaction();
+		});
+
+		REQUIRE(nullptr == Blob);
+
+		// When we abort we'll actually return the allocated memory to the allocator!
+		REQUIRE(1 == Allocator.GetNumFree());
+	}
+
+	SECTION("Transact(Transact(Allocate, Abort), Allocate)")
+	{
+		ALockFreeFixedSizeAllocator Allocator;
+		void* Blob = nullptr;
+
+		AutoRTFM::Testing::Commit([&]
+		{
+			AutoRTFM::Testing::Abort([&]
+			{
+				Blob = Allocator.Allocate();
+				AutoRTFM::AbortTransaction();
+			});
+			Blob = Allocator.Allocate();
+		});
+
+		REQUIRE(nullptr != Blob);
+
+		Allocator.Free(Blob);
+
+		// Check that the inner abort did eagerly return the allocation to the allocator,
+		// and that the outer allocate reused that allocation.
+		REQUIRE(1 == Allocator.GetNumFree());
+	}
+
+	SECTION("Allocate, Transact(Free)")
+	{
+		ALockFreeFixedSizeAllocator Allocator;
+		void* Blob = Allocator.Allocate();
+		IntegerType NumUsed = 0;
+		IntegerType NumFree = 0;
+
+		AutoRTFM::Testing::Commit([&]
+		{
+			Allocator.Free(Blob);
+			NumUsed = Allocator.GetNumUsed();
+			NumFree = Allocator.GetNumFree();
+		});
+
+		// Even though we freed in the transaction, we won't do the free until on-commit so
+		// the query of these within the transaction will not be updated!
+		REQUIRE(1 == NumUsed);
+		REQUIRE(0 == NumFree);
+
+		// But after the transaction they will return the correct values.
+		REQUIRE(0 == Allocator.GetNumUsed());
+		REQUIRE(1 == Allocator.GetNumFree());
+	}
+
+	SECTION("Allocate, Transact(Free, Abort)")
+	{
+		ALockFreeFixedSizeAllocator Allocator;
+		void* Blob = Allocator.Allocate();
+
+		AutoRTFM::Testing::Abort([&]
+		{
+			Allocator.Free(Blob);
+			AutoRTFM::AbortTransaction();
+		});
+
+		REQUIRE(1 == Allocator.GetNumUsed());
+		REQUIRE(0 == Allocator.GetNumFree());
+
+		Allocator.Free(Blob);
+	}
+
+	SECTION("Transact(Allocate, Free)")
+	{
+		ALockFreeFixedSizeAllocator Allocator;
+
+		AutoRTFM::Testing::Commit([&]
+			{
+				void* Blob = Allocator.Allocate();
+				Allocator.Free(Blob);
+			});
+
+		REQUIRE(0 == Allocator.GetNumUsed());
+		REQUIRE(1 == Allocator.GetNumFree());
+	}
+
+	SECTION("Transact(Allocate, Free, Abort)")
+	{
+		ALockFreeFixedSizeAllocator Allocator;
+
+		AutoRTFM::Testing::Abort([&]
+		{
+			void* Blob = Allocator.Allocate();
+			Allocator.Free(Blob);
+			AutoRTFM::AbortTransaction();
+		});
+
+		// Even though we aborted, the allocation will be cached in the allocator!
+		REQUIRE(0 == Allocator.GetNumUsed());
+		REQUIRE(1 == Allocator.GetNumFree());
+	}
+
+	SECTION("Allocate, Allocate, Allocate, Free, Free, Free, Transact(Trim)")
+	{
+		ALockFreeFixedSizeAllocator Allocator;
+
+		void* Blob0 = Allocator.Allocate();
+		void* Blob1 = Allocator.Allocate();
+		void* Blob2 = Allocator.Allocate();
+		Allocator.Free(Blob0);
+		Allocator.Free(Blob1);
+		Allocator.Free(Blob2);
+
+		REQUIRE(3 == Allocator.GetNumFree());
+
+		AutoRTFM::Testing::Commit([&]
+			{
+				Allocator.Trim();
+			});
+
+		REQUIRE(0 == Allocator.GetNumFree());
+	}
+
+	SECTION("Allocate, Allocate, Allocate, Free, Free, Free, Transact(Trim, Abort)")
+	{
+		ALockFreeFixedSizeAllocator Allocator;
+
+		void* Blob0 = Allocator.Allocate();
+		void* Blob1 = Allocator.Allocate();
+		void* Blob2 = Allocator.Allocate();
+		Allocator.Free(Blob0);
+		Allocator.Free(Blob1);
+		Allocator.Free(Blob2);
+
+		REQUIRE(3 == Allocator.GetNumFree());
+
+		AutoRTFM::Testing::Abort([&]
+			{
+				Allocator.Trim();
+				AutoRTFM::AbortTransaction();
+			});
+
+		// We aborted so the trim did not happen!
+		REQUIRE(3 == Allocator.GetNumFree());
+	}
+
+	SECTION("Allocate, Allocate, Allocate, Transact(Trim, Allocate, Free, Free, Free, Free)")
+	{
+		ALockFreeFixedSizeAllocator Allocator;
+
+		void* Blob0 = Allocator.Allocate();
+		void* Blob1 = Allocator.Allocate();
+		void* Blob2 = Allocator.Allocate();
+
+		AutoRTFM::Testing::Commit([&]
+		{
+			Allocator.Trim();
+
+			void* Blob = Allocator.Allocate();
+			Allocator.Free(Blob);
+			Allocator.Free(Blob0);
+			Allocator.Free(Blob1);
+			Allocator.Free(Blob2);
+		});
+
+		REQUIRE(4 == Allocator.GetNumFree());
+	}
+	
+	SECTION("Transact(GetNumFree), Allocate, Free, Transact(GetNumFree)")
+	{
+		ALockFreeFixedSizeAllocator Allocator;
+		IntegerType NumFree = 0;
+
+		AutoRTFM::Commit([&]
+		{
+			NumFree = Allocator.GetNumFree();
+		});
+
+		REQUIRE(0 == NumFree);
+
+		void* const Blob = Allocator.Allocate();
+		Allocator.Free(Blob);
+
+		AutoRTFM::Commit([&]
+		{
+			NumFree = Allocator.GetNumFree();
+		});
+
+		REQUIRE(1 == NumFree);
+	}
+
+	SECTION("Transact(GetNumUsed), Allocate, Transact(GetNumUsed), Free")
+	{
+		ALockFreeFixedSizeAllocator Allocator;
+		IntegerType NumUsed = 0;
+
+		AutoRTFM::Commit([&]
+		{
+			NumUsed = Allocator.GetNumUsed();
+		});
+
+		REQUIRE(0 == NumUsed);
+
+		void* const Blob = Allocator.Allocate();
+
+		AutoRTFM::Commit([&]
+		{
+			NumUsed = Allocator.GetNumUsed();
+		});
+
+		REQUIRE(1 == NumUsed);
+
+		Allocator.Free(Blob);
+	}
+}
diff --git a/Engine/Source/Runtime/Core/Private/Misc/MemStack.cpp b/Engine/Source/Runtime/Core/Private/Misc/MemStack.cpp
--- a/Engine/Source/Runtime/Core/Private/Misc/MemStack.cpp
+++ b/Engine/Source/Runtime/Core/Private/Misc/MemStack.cpp
@@ -267,12 +267,12 @@ void FPageAllocator::FreeSmall(void *Mem)
 
 uint64 FPageAllocator::BytesUsed()
 {
-	return uint64(TheAllocator.GetNumUsed().GetValue()) * PageSize;
+	return uint64(TheAllocator.GetNumUsed()) * PageSize;
 }
 
 uint64 FPageAllocator::BytesFree()
 {
-	return uint64(TheAllocator.GetNumFree().GetValue()) * PageSize;
+	return uint64(TheAllocator.GetNumFree()) * PageSize;
 }
 
 #endif
diff --git a/Engine/Source/Runtime/Core/Public/AutoRTFM/AutoRTFM.h b/Engine/Source/Runtime/Core/Public/AutoRTFM/AutoRTFM.h
--- a/Engine/Source/Runtime/Core/Public/AutoRTFM/AutoRTFM.h
+++ b/Engine/Source/Runtime/Core/Public/AutoRTFM/AutoRTFM.h
@@ -837,7 +837,13 @@ template<typename TFunctor> static UE_AUTORTFM_FORCEINLINE void OnAbort(TFunctor
 		ForTheRuntime::OnAbortInternal(std::forward<TFunctor>(Work));
 	}
 }
+#else
+// Have some work happen when this transaction aborts. If this is called
+// outside a transaction or from an open nest then the work is ignored.
+template<typename TFunctor> static UE_AUTORTFM_FORCEINLINE void OnAbort(TFunctor&&) {}
+#endif
 
+#if UE_AUTORTFM
 // Register a handler for transaction abort. Takes a key parameter so that
 // the handler can be unregistered (see PopOnAbortHandler). This is useful
 // for scoped mutations that need an abort handler present unless execution
@@ -849,7 +855,15 @@ template<typename TFunctor> static UE_AUTORTFM_FORCEINLINE void PushOnAbortHandl
 		ForTheRuntime::PushOnAbortHandlerInternal(Key, std::forward<TFunctor>(Work));
 	}
 }
+#else
+// Register a handler for transaction abort. Takes a key parameter so that
+// the handler can be unregistered (see PopOnAbortHandler). This is useful
+// for scoped mutations that need an abort handler present unless execution
+// reaches the end of the relevant scope.
+template<typename TFunctor> static UE_AUTORTFM_FORCEINLINE void PushOnAbortHandler(const void* Key, TFunctor&&) {}
+#endif
 
+#if UE_AUTORTFM
 // Unregister the most recently pushed handler (via PushOnAbortHandler) for the given key.
 static UE_AUTORTFM_FORCEINLINE void PopOnAbortHandler(const void* Key)
 {
@@ -858,7 +872,12 @@ static UE_AUTORTFM_FORCEINLINE void PopOnAbortHandler(const void* Key)
 		ForTheRuntime::PopOnAbortHandlerInternal(Key);
 	}
 }
+#else
+// Unregister the most recently pushed handler (via PushOnAbortHandler) for the given key.
+static UE_AUTORTFM_FORCEINLINE void PopOnAbortHandler(const void* Key) {}
+#endif
 
+#if UE_AUTORTFM
 // Unregister all pushed handlers (via PushOnAbortHandler) for the given key.
 static UE_AUTORTFM_FORCEINLINE void PopAllOnAbortHandlers(const void* Key)
 {
@@ -868,20 +887,6 @@ static UE_AUTORTFM_FORCEINLINE void PopAllOnAbortHandlers(const void* Key)
 	}
 }
 #else
-// Have some work happen when this transaction aborts. If this is called
-// outside a transaction or from an open nest then the work is ignored.
-template<typename TFunctor> static UE_AUTORTFM_FORCEINLINE void OnAbort(TFunctor&&) {}
-
-// Register a handler for transaction abort. Takes a key parameter so that
-// the handler can be unregistered (see PopOnAbortHandler). This is useful
-// for scoped mutations that need an abort handler present unless execution
-// reaches the end of the relevant scope.
-template<typename TFunctor> static UE_AUTORTFM_FORCEINLINE void PushOnAbortHandler(const void* Key, TFunctor&&) {}
-
-// Unregister the most recently pushed handler (via PushOnAbortHandler) for the given key.
-static UE_AUTORTFM_FORCEINLINE void PopOnAbortHandler(const void* Key) {}
-
-
 // Unregister all pushed handlers (via PushOnAbortHandler) for the given key.
 static UE_AUTORTFM_FORCEINLINE void PopAllOnAbortHandlers(const void* Key) {}
 #endif
diff --git a/Engine/Source/Runtime/Core/Public/AutoRTFM/AutoRTFMTesting.h b/Engine/Source/Runtime/Core/Public/AutoRTFM/AutoRTFMTesting.h
new file mode 100644
--- /dev/null
+++ b/Engine/Source/Runtime/Core/Public/AutoRTFM/AutoRTFMTesting.h
@@ -0,0 +1,39 @@
+// Copyright Epic Games, Inc. All Rights Reserved.
+
+#pragma once
+
+#include "AutoRTFM.h"
+
+#if defined(UE_AUTORTFM_STANDALONE)
+#include <assert.h>
+#define UE_AUTORTFM_CHECK(x) assert(x)
+#else
+#include <Misc/AssertionMacros.h>
+#define UE_AUTORTFM_CHECK(x) check(x)
+#endif
+
+namespace AutoRTFM
+{
+namespace Testing
+{
+	// Run the callback in a transaction like Transact, but abort program
+	// execution if the result is anything other than autortfm_committed.
+	// Useful for testing.
+	template<typename TFunctor>
+	static UE_AUTORTFM_FORCEINLINE void Commit(const TFunctor& Functor)
+	{
+		const ETransactionResult Result = Transact(Functor);
+		UE_AUTORTFM_CHECK(ETransactionResult::Committed == Result);
+	}
+
+	// Run the callback in a transaction like Transact, but abort program
+	// execution if the result is anything other than abort.
+	template<typename TFunctor>
+	static UE_AUTORTFM_FORCEINLINE void Abort(const TFunctor& Functor)
+	{
+		const ETransactionResult Result = Transact(Functor);
+		UE_AUTORTFM_CHECK(ETransactionResult::Committed != Result);
+	}
+} // namespace Testing
+
+} // namespace AutoRTFM
diff --git a/Engine/Source/Runtime/Core/Public/Containers/LockFreeFixedSizeAllocator.h b/Engine/Source/Runtime/Core/Public/Containers/LockFreeFixedSizeAllocator.h
--- a/Engine/Source/Runtime/Core/Public/Containers/LockFreeFixedSizeAllocator.h
+++ b/Engine/Source/Runtime/Core/Public/Containers/LockFreeFixedSizeAllocator.h
@@ -7,7 +7,6 @@
 #include "Misc/AssertionMacros.h"
 #include "HAL/UnrealMemory.h"
 #include "Misc/NoopCounter.h"
-#include "Misc/TransactionallySafeCounter.h"
 #include "Containers/LockFreeList.h"
 
 /**
@@ -198,12 +197,17 @@ template<int32 SIZE, int TPaddingForCacheContention, typename TTrackingCounter =
 class TLockFreeFixedSizeAllocator
 {
 public:
+	TLockFreeFixedSizeAllocator()
+	{
+		checkf(AutoRTFM::IsClosed() || !AutoRTFM::IsOnCurrentTransactionStack(this), TEXT("Not allowed to construct a stack local within a transaction."));
+	}
+
 	/** Destructor, returns all memory via FMemory::Free **/
 	~TLockFreeFixedSizeAllocator()
 	{
-		UE_AUTORTFM_OPEN
+		UE_AUTORTFM_ONCOMMIT(this)
 		{
-			check(!NumUsed.GetValue());
+			check(GetNumUsed() == 0);
 			Trim();
 		};
 	}
@@ -216,25 +220,32 @@ public:
 	 */
 	void* Allocate()
 	{
-		NumUsed.Increment();
+		void* Memory = nullptr;
 
-		// When we are in an AutoRTFM transaction, it is simpler and faster to allocate a new block instead
-		// of reusing an existing one, since we don't need to track changes on newly allocated memory.
-		// If the transaction is aborted, AutoRTFM should clean it up for us automatically.
-		if (AutoRTFM::IsClosed())
-		{
-			return FMemory::Malloc(SIZE);
-		}
-		// Outside of a transaction, we prefer to reuse blocks from the FreeList.
-		void* Memory = FreeList.Pop();
-		if (Memory)
+		// We need the allocation to happen (so we get a valid `Memory`)
+		// pointer to use in the transaction. So we do the actual meat
+		// of the allocation in the open.
+		UE_AUTORTFM_OPEN
 		{
-			NumFree.Decrement();
-		}
-		else
+			NumUsed.Increment();
+			Memory = FreeList.Pop();
+			if (Memory)
+			{
+				NumFree.Decrement();
+			}
+			else
+			{
+				Memory = FMemory::Malloc(SIZE);
+			}
+		};
+
+		// But if we abort the transaction, we need to return the memory
+		// region to the allocator (otherwise we'll leak).
+		UE_AUTORTFM_ONABORT(this, Memory)
 		{
-			Memory = FMemory::Malloc(SIZE);
-		}
+			this->Free(Memory);
+		};
+		
 		return Memory;
 	}
 
@@ -246,21 +257,15 @@ public:
 	 */
 	void Free(void* Item)
 	{
-		NumUsed.Decrement();
-
-		if (AutoRTFM::IsClosed())
-		{
-			// When we are in an AutoRTFM transaction, it is simpler to free blocks directly.
-			// Frees are deferred until the transaction is complete, and we don't need to
-			// worry about leaks occurring if transactions are aborted at inopportune times.
-			FMemory::Free(Item);
-		}
-		else
+		// We defer actually returning `Item` until on-commit time.
+		// If we didn't do this an aborting transaction would not
+		// be able to be undone.
+		UE_AUTORTFM_ONCOMMIT(this, Item)
 		{
-			// Outside of a transaction, we push blocks onto the FreeList for reuse.
+			NumUsed.Decrement();
 			FreeList.Push(Item);
 			NumFree.Increment();
-		}
+		};
 	}
 
 	/**
@@ -268,7 +273,10 @@ public:
 	*/
 	void Trim()
 	{
-		UE_AUTORTFM_OPEN
+		// Similar to `Free`, we need to only do the destructive
+		// trim operation at commit time - otherwise we'll wouldn't
+		// be able to undo it.
+		UE_AUTORTFM_ONCOMMIT(this)
 		{
 			while (void* Mem = FreeList.Pop())
 			{
@@ -284,9 +292,9 @@ public:
 	 * @return Number of used memory blocks.
 	 * @see GetNumFree
 	 */
-	const TTrackingCounter& GetNumUsed() const
+	typename TTrackingCounter::IntegerType GetNumUsed() const
 	{
-		return NumUsed.GetCounter();
+		return AutoRTFM::Open([&] { return NumUsed.GetValue(); });
 	}
 
 	/**
@@ -295,18 +303,19 @@ public:
 	 * @return Number of unused memory blocks.
 	 * @see GetNumUsed
 	 */
-	const TTrackingCounter& GetNumFree() const
+	typename TTrackingCounter::IntegerType GetNumFree() const
 	{
-		return NumFree;
+		return AutoRTFM::Open([&] { return NumFree.GetValue(); });
 	}
 
 private:
+	UE_NONCOPYABLE(TLockFreeFixedSizeAllocator)
 
 	/** Lock free list of free memory blocks. */
 	TLockFreePointerListUnordered<void, TPaddingForCacheContention> FreeList;
 
 	/** Total number of blocks outstanding and not in the free list. */
-	TTransactionallySafeCounter<TTrackingCounter> NumUsed;
+	TTrackingCounter NumUsed;
 
 	/** Total number of blocks in the free list. */
 	TTrackingCounter NumFree;
diff --git a/Engine/Source/Runtime/Core/Public/Misc/TransactionallySafeCounter.h b/Engine/Source/Runtime/Core/Public/Misc/TransactionallySafeCounter.h
deleted file mode 100644
diff --git a/Engine/Source/Runtime/Core/Tests/Containers/AutoRTFMLockFreeFixedSizeAllocatorTest.cpp b/Engine/Source/Runtime/Core/Tests/Containers/AutoRTFMLockFreeFixedSizeAllocatorTest.cpp
deleted file mode 100644
diff --git a/Engine/Source/Runtime/Core/Tests/HAL/AutoRTFMTransactionallySafeCounterTest.cpp b/Engine/Source/Runtime/Core/Tests/HAL/AutoRTFMTransactionallySafeCounterTest.cpp
deleted file mode 100644
diff --git a/Engine/Source/Runtime/CoreUObject/Public/AssetRegistry/AssetDataTagMapSerializationDetails.h b/Engine/Source/Runtime/CoreUObject/Public/AssetRegistry/AssetDataTagMapSerializationDetails.h
--- a/Engine/Source/Runtime/CoreUObject/Public/AssetRegistry/AssetDataTagMapSerializationDetails.h
+++ b/Engine/Source/Runtime/CoreUObject/Public/AssetRegistry/AssetDataTagMapSerializationDetails.h
@@ -86,9 +86,7 @@ namespace FixedTagPrivate
 
 		explicit FStore(uint32 InIndex) : Index(InIndex)
 		{
-			// `IsClosed` will fold away at compile time to true|false meaning this check will only impact
-			// transactionally compiled code!
-			check(AutoRTFM::IsClosed() || !AutoRTFM::IsOnCurrentTransactionStack(this));
+			checkf(AutoRTFM::IsClosed() || !AutoRTFM::IsOnCurrentTransactionStack(this), TEXT("Not allowed to construct a stack local within a transaction."));
 		}
 
 		~FStore();
diff --git a/Engine/Source/Runtime/MassEntity/Private/MassEntityManagerStorage.cpp b/Engine/Source/Runtime/MassEntity/Private/MassEntityManagerStorage.cpp
--- a/Engine/Source/Runtime/MassEntity/Private/MassEntityManagerStorage.cpp
+++ b/Engine/Source/Runtime/MassEntity/Private/MassEntityManagerStorage.cpp
@@ -412,6 +412,8 @@ namespace UE::Mass
 			}
 
 			EntityIndex = EntityFreeIndexList.Pop(EAllowShrinking::No);
+
+			++EntityCount;
 		}
 
 		FEntityData& EntityData = LookupEntity(EntityIndex);
@@ -422,8 +424,6 @@ namespace UE::Mass
 		EntityData.bIsAllocated = 1;
 		int32 SerialNumber = EntityData.GetSerialNumber();
 
-		EntityCount.Increment();
-	
 		FMassEntityHandle Handle;
 		Handle.SerialNumber = SerialNumber;
 		Handle.Index = EntityIndex;
@@ -468,11 +468,10 @@ namespace UE::Mass
 			}
 			
 			CountAdded += CountToProcess;
+			EntityCount += CountToProcess;
 			CountLeft -= CountToProcess;
 		}
 
-		EntityCount.Add(CountAdded);
-
 		return CountAdded;
 	}
 
@@ -526,7 +525,10 @@ namespace UE::Mass
 		// Free any remaining handles
 		FreeRunOfHandles();
 
-		EntityCount.Subtract(DeallocateCount);
+		{
+			UE::TUniqueLock FreeListLock(FreeListMutex);
+			EntityCount -= DeallocateCount;
+		}
 	
 		return DeallocateCount;
 	}
@@ -556,10 +558,10 @@ namespace UE::Mass
 			{
 				EntityFreeIndexList.Add(Handle.Index);
 			}
+
+			EntityCount -= Handles.Num();
 		}
 
-		EntityCount.Subtract(Handles.Num());
-	
 		return Handles.Num();
 	}
 
diff --git a/Engine/Source/Runtime/MassEntity/Public/MassEntityManagerStorage.h b/Engine/Source/Runtime/MassEntity/Public/MassEntityManagerStorage.h
--- a/Engine/Source/Runtime/MassEntity/Public/MassEntityManagerStorage.h
+++ b/Engine/Source/Runtime/MassEntity/Public/MassEntityManagerStorage.h
@@ -3,8 +3,6 @@
 
 #include "Async/Mutex.h"
 #include "Containers/ChunkedArray.h"
-#include "HAL/ThreadSafeCounter.h"
-#include "Misc/TransactionallySafeCounter.h"
 #include "Misc/TransactionallySafeMutex.h"
 #include "Misc/TVariant.h"
 #include "Templates/SharedPointer.h"
@@ -219,8 +217,8 @@ namespace UE::Mass
 		// @return whether the operation was successful. Will return false when OOM
 		bool AddPage();
 
-		// Number of allocated Entities
-		TTransactionallySafeCounter<FThreadSafeCounter> EntityCount = 0;
+		// Number of allocated Entities (only used for viewing in the debugger).
+		uint32 EntityCount = 0;
 		uint32 MaximumEntityCountShift = 0;
 		uint32 MaxEntitiesPerPage = 0;
 		uint32 PageCount = 0;
