[AutoRTFM] Only write to FUObjectThreadContext in the open

Writing to this in both the open and closed may lead to corruption.

diff --git a/Engine/Source/Programs/AutoRTFMTests/Classes/MyAutoRTFMTestObject.h b/Engine/Source/Programs/AutoRTFMTests/Classes/MyAutoRTFMTestObject.h
--- a/Engine/Source/Programs/AutoRTFMTests/Classes/MyAutoRTFMTestObject.h
+++ b/Engine/Source/Programs/AutoRTFMTests/Classes/MyAutoRTFMTestObject.h
@@ -9,17 +9,11 @@ class UMyAutoRTFMTestObject : public UObject
 	GENERATED_BODY()
 
 public:
-	UMyAutoRTFMTestObject(const FObjectInitializer& ObjectInitializer = FObjectInitializer::Get()) : Value(42)
-	{
-		UObject* const Obj = ObjectInitializer.GetObj();
-		UObject* const Outer = Obj->GetOuter();
+	using FConstructorCallback = void(const FObjectInitializer&, UMyAutoRTFMTestObject&);
 
-		if (Outer->IsA<UMyAutoRTFMTestObject>())
-		{
-			UMyAutoRTFMTestObject* const OuterAsType = static_cast<UMyAutoRTFMTestObject*>(Outer);
-			OuterAsType->Value += 13;
-		}
-	}
+	UMyAutoRTFMTestObject(const FObjectInitializer& ObjectInitializer = FObjectInitializer::Get());
 
 	int Value;
+
+	static FConstructorCallback* ConstructorCallback;
 };
diff --git a/Engine/Source/Programs/AutoRTFMTests/Private/MyAutoRTFMTestObject.cpp b/Engine/Source/Programs/AutoRTFMTests/Private/MyAutoRTFMTestObject.cpp
new file mode 100644
--- /dev/null
+++ b/Engine/Source/Programs/AutoRTFMTests/Private/MyAutoRTFMTestObject.cpp
@@ -0,0 +1,22 @@
+// Copyright Epic Games, Inc. All Rights Reserved.
+
+#include "MyAutoRTFMTestObject.h"
+
+UMyAutoRTFMTestObject::FConstructorCallback* UMyAutoRTFMTestObject::ConstructorCallback = nullptr;
+
+UMyAutoRTFMTestObject::UMyAutoRTFMTestObject(const FObjectInitializer& ObjectInitializer /* = FObjectInitializer::Get() */) : Value(42)
+{
+	UObject* const Obj = ObjectInitializer.GetObj();
+	UObject* const Outer = Obj->GetOuter();
+
+	if (Outer->IsA<UMyAutoRTFMTestObject>())
+	{
+		UMyAutoRTFMTestObject* const OuterAsType = static_cast<UMyAutoRTFMTestObject*>(Outer);
+		OuterAsType->Value += 13;
+	}
+
+	if (ConstructorCallback)
+	{
+		ConstructorCallback(ObjectInitializer, *this);
+	}
+}
diff --git a/Engine/Source/Programs/AutoRTFMTests/Private/UObjectTests.cpp b/Engine/Source/Programs/AutoRTFMTests/Private/UObjectTests.cpp
--- a/Engine/Source/Programs/AutoRTFMTests/Private/UObjectTests.cpp
+++ b/Engine/Source/Programs/AutoRTFMTests/Private/UObjectTests.cpp
@@ -1,7 +1,9 @@
 // Copyright Epic Games, Inc. All Rights Reserved.
 
+#include "AutoRTFM/ScopedGuard.h"
 #include "GenericPlatform/GenericPlatformMisc.h"
 #include "HAL/MallocLeakDetection.h"
+#include "UObject/UObjectThreadContext.h"
 
 THIRD_PARTY_INCLUDES_START
 #include "Catch2Includes.h"
@@ -345,3 +347,193 @@ TEST_CASE("UObject.AtomicallyClearFlags")
 	REQUIRE(AutoRTFM::ETransactionResult::Committed == Result);
 	REQUIRE(Object->HasAnyFlags(OldFlags) & !Object->HasAllFlags(FlagsToClear));
 }
+
+// Tests that constructing a UObject in both the open and closed doesn't result
+// in a corrupt FUObjectThreadContext. See SOL-7131.
+TEST_CASE("UObject.FUObjectThreadContext")
+{
+	AutoRTFM::TScopedGuard<UMyAutoRTFMTestObject::FConstructorCallback*> CallbackScope(UMyAutoRTFMTestObject::ConstructorCallback, nullptr);
+
+	struct Fns
+	{
+		static void CreateObjectWithCtor(UMyAutoRTFMTestObject::FConstructorCallback* Ctor)
+		{
+			AutoRTFM::Open([&] { UMyAutoRTFMTestObject::ConstructorCallback = Ctor; });
+			NewObject<UMyAutoRTFMTestObject>();
+		}
+		
+		static void CtorCreateInnerClosed(const FObjectInitializer& ObjectInitializer, UMyAutoRTFMTestObject& Object)
+		{
+			AutoRTFM::Close([&]
+			{
+				REQUIRE(1 == FUObjectThreadContext::Get().IsInConstructor);
+				CreateObjectWithCtor(nullptr);
+				REQUIRE(1 == FUObjectThreadContext::Get().IsInConstructor);
+			});
+		}
+		
+		static void CtorCreateInnerTransact(const FObjectInitializer& ObjectInitializer, UMyAutoRTFMTestObject& Object)
+		{
+			AutoRTFM::Transact([&]
+			{
+				REQUIRE(1 == FUObjectThreadContext::Get().IsInConstructor);
+				CreateObjectWithCtor(nullptr);
+				REQUIRE(1 == FUObjectThreadContext::Get().IsInConstructor);
+			});
+		}
+
+		static void CtorAbort(const FObjectInitializer& ObjectInitializer, UMyAutoRTFMTestObject& Object)
+		{
+			REQUIRE(1 == FUObjectThreadContext::Get().IsInConstructor);
+			AutoRTFM::Close([&]
+			{
+				AutoRTFM::AbortTransaction();
+				FAIL(/* unreachable */);
+			});
+		}
+	};
+
+	SECTION("Transact(UObjectCtor(Abort))")
+	{
+		AutoRTFM::ETransactionResult Result = AutoRTFM::Transact([&]
+			{
+				REQUIRE(AutoRTFM::IsClosed());
+				REQUIRE(0 == FUObjectThreadContext::Get().IsInConstructor);
+				
+				Fns::CreateObjectWithCtor(Fns::CtorAbort);
+				FAIL(/* unreachable */);
+			});
+
+		REQUIRE(AutoRTFM::ETransactionResult::AbortedByRequest == Result);
+		REQUIRE(0 == FUObjectThreadContext::Get().IsInConstructor);
+	}
+
+	SECTION("Transact(UObjectCtor(), Abort)")
+	{
+		AutoRTFM::ETransactionResult Result = AutoRTFM::Transact([&]
+			{
+				REQUIRE(AutoRTFM::IsClosed());
+				REQUIRE(0 == FUObjectThreadContext::Get().IsInConstructor);
+				
+				NewObject<UMyAutoRTFMTestObject>();
+
+				REQUIRE(0 == FUObjectThreadContext::Get().IsInConstructor);
+
+				AutoRTFM::AbortTransaction();
+				FAIL(/* unreachable */);
+			});
+
+		REQUIRE(AutoRTFM::ETransactionResult::AbortedByRequest == Result);
+		REQUIRE(0 == FUObjectThreadContext::Get().IsInConstructor);
+	}
+
+	SECTION("Transact(Open(UObjectCtor(Abort)), UObjectCtor(Abort))")
+	{
+		AutoRTFM::ETransactionResult Result = AutoRTFM::Transact([&]
+			{
+				REQUIRE(AutoRTFM::IsClosed());
+				REQUIRE(0 == FUObjectThreadContext::Get().IsInConstructor);
+				
+				AutoRTFM::Open([&]
+				{
+					AutoRTFM::AbortTransaction();
+				});
+
+				REQUIRE(0 == FUObjectThreadContext::Get().IsInConstructor);
+
+				Fns::CreateObjectWithCtor(Fns::CtorAbort);
+				FAIL(/* unreachable */);
+			});
+
+		REQUIRE(AutoRTFM::ETransactionResult::AbortedByRequest == Result);
+		REQUIRE(0 == FUObjectThreadContext::Get().IsInConstructor);
+	}
+
+	SECTION("Transact(Open(UObjectCtor(Transact(UObjectCtor))), UObjectCtor(Abort))")
+	{
+		AutoRTFM::ETransactionResult Result = AutoRTFM::Transact([&]
+			{
+				REQUIRE(AutoRTFM::IsClosed());
+				REQUIRE(0 == FUObjectThreadContext::Get().IsInConstructor);
+				
+				AutoRTFM::Open([&]
+				{
+					Fns::CreateObjectWithCtor(Fns::CtorCreateInnerTransact);
+				});
+
+				REQUIRE(0 == FUObjectThreadContext::Get().IsInConstructor);
+
+				Fns::CreateObjectWithCtor(Fns::CtorAbort);
+				FAIL(/* unreachable */);
+			});
+
+		REQUIRE(AutoRTFM::ETransactionResult::AbortedByRequest == Result);
+		REQUIRE(0 == FUObjectThreadContext::Get().IsInConstructor);
+	}
+
+	SECTION("Transact(Open(UObjectCtor(Transact(UObjectCtor))), Abort)")
+	{
+		AutoRTFM::ETransactionResult Result = AutoRTFM::Transact([&]
+			{
+				REQUIRE(AutoRTFM::IsClosed());
+				REQUIRE(0 == FUObjectThreadContext::Get().IsInConstructor);
+				
+				AutoRTFM::Open([&]
+				{
+					Fns::CreateObjectWithCtor(Fns::CtorCreateInnerTransact);
+				});
+
+				REQUIRE(0 == FUObjectThreadContext::Get().IsInConstructor);
+
+				AutoRTFM::AbortTransaction();
+				FAIL(/* unreachable */);
+			});
+
+		REQUIRE(AutoRTFM::ETransactionResult::AbortedByRequest == Result);
+		REQUIRE(0 == FUObjectThreadContext::Get().IsInConstructor);
+	}
+
+	SECTION("Transact(Open(UObjectCtor(Close(UObjectCtor))), UObjectCtor(Abort))")
+	{
+		AutoRTFM::ETransactionResult Result = AutoRTFM::Transact([&]
+			{
+				REQUIRE(AutoRTFM::IsClosed());
+				REQUIRE(0 == FUObjectThreadContext::Get().IsInConstructor);
+				
+				AutoRTFM::Open([&]
+				{
+					Fns::CreateObjectWithCtor(Fns::CtorCreateInnerClosed);
+				});
+
+				REQUIRE(0 == FUObjectThreadContext::Get().IsInConstructor);
+
+				Fns::CreateObjectWithCtor(Fns::CtorAbort);
+				FAIL(/* unreachable */);
+			});
+
+		REQUIRE(AutoRTFM::ETransactionResult::AbortedByRequest == Result);
+		REQUIRE(0 == FUObjectThreadContext::Get().IsInConstructor);
+	}
+	
+	SECTION("Transact(Open(UObjectCtor(Close(UObjectCtor))), Abort)")
+	{
+		AutoRTFM::ETransactionResult Result = AutoRTFM::Transact([&]
+			{
+				REQUIRE(AutoRTFM::IsClosed());
+				REQUIRE(0 == FUObjectThreadContext::Get().IsInConstructor);
+				
+				AutoRTFM::Open([&]
+				{
+					Fns::CreateObjectWithCtor(Fns::CtorCreateInnerClosed);
+				});
+
+				REQUIRE(0 == FUObjectThreadContext::Get().IsInConstructor);
+
+				AutoRTFM::AbortTransaction();
+				FAIL(/* unreachable */);
+			});
+
+		REQUIRE(AutoRTFM::ETransactionResult::AbortedByRequest == Result);
+		REQUIRE(0 == FUObjectThreadContext::Get().IsInConstructor);
+	}
+}
diff --git a/Engine/Source/Runtime/CoreUObject/Private/UObject/UObjectGlobals.cpp b/Engine/Source/Runtime/CoreUObject/Private/UObject/UObjectGlobals.cpp
--- a/Engine/Source/Runtime/CoreUObject/Private/UObject/UObjectGlobals.cpp
+++ b/Engine/Source/Runtime/CoreUObject/Private/UObject/UObjectGlobals.cpp
@@ -3699,7 +3699,19 @@ void UObject::PostReinitProperties()
 void UObject::PostInitProperties()
 {
 #if !(UE_BUILD_SHIPPING || UE_BUILD_TEST)
-	FUObjectThreadContext::Get().PostInitPropertiesCheck.Push(this);
+	FUObjectThreadContext& ThreadContext = FUObjectThreadContext::Get();
+	// Always write to FUObjectThreadContext from the open, as writing to a 
+	// memory location in both the open and closed within the same transaction
+	// can lead to state corruption.
+	UE_AUTORTFM_OPEN
+	{
+		ThreadContext.PostInitPropertiesCheck.Push(this);
+	};
+	AutoRTFM::PushOnAbortHandler(&ThreadContext.PostInitPropertiesCheck, [this]
+	{ 
+		UObject* Popped = FUObjectThreadContext::Get().PostInitPropertiesCheck.Pop(EAllowShrinking::No);
+		check(this == Popped);
+	});
 #endif
 
 	GetClass()->CreatePersistentUberGraphFrame(this, true);
@@ -3791,11 +3803,30 @@ FObjectInitializer::FObjectInitializer(UObject* InObj, UObject* InObjectArchetyp
 void FObjectInitializer::Construct_Internal()
 {
 	FUObjectThreadContext& ThreadContext = FUObjectThreadContext::Get();
-	// Mark we're in the constructor now.
-	ThreadContext.IsInConstructor++;
-	LastConstructedObject = ThreadContext.ConstructedObject;
-	ThreadContext.ConstructedObject = Obj;
-	ThreadContext.PushInitializer(this);
+	// Always write to FUObjectThreadContext from the open, as writing to a 
+	// memory location in both the open and closed within the same transaction
+	// can lead to state corruption.
+	UE_AUTORTFM_OPEN
+	{
+		// Mark we're in the constructor now.
+		ThreadContext.IsInConstructor++;
+		LastConstructedObject = ThreadContext.ConstructedObject;
+		ThreadContext.ConstructedObject = Obj;
+		ThreadContext.PushInitializer(this);
+	};
+	AutoRTFM::PushOnAbortHandler(&ThreadContext.IsInConstructor, [this]
+	{ 
+		FUObjectThreadContext& ThreadContext = FUObjectThreadContext::Get();
+		ThreadContext.IsInConstructor--;
+		check(ThreadContext.IsInConstructor >= 0);
+		ThreadContext.ConstructedObject = LastConstructedObject;
+	});
+	AutoRTFM::PushOnAbortHandler(this, [this]
+	{ 
+		FUObjectThreadContext& ThreadContext = FUObjectThreadContext::Get();
+		check(ThreadContext.TopInitializer() == this);
+		ThreadContext.PopInitializer();
+	});
 
 	if (Obj && GetAllowNativeComponentClassOverrides())
 	{
@@ -3852,10 +3883,17 @@ FObjectInitializer::~FObjectInitializer()
 	if (!bIsDeferredInitializer)
 	{
 #endif // USE_CIRCULAR_DEPENDENCY_LOAD_DEFERRING
-		// Let the FObjectFinders know we left the constructor.
-		ThreadContext.IsInConstructor--;
-		check(ThreadContext.IsInConstructor >= 0);
-		ThreadContext.ConstructedObject = LastConstructedObject;
+		// Always write to FUObjectThreadContext from the open, as writing to a 
+		// memory location in both the open and closed within the same transaction
+		// can lead to state corruption.
+		AutoRTFM::PopOnAbortHandler(&ThreadContext.IsInConstructor);
+		UE_AUTORTFM_OPEN
+		{
+			// Let the FObjectFinders know we left the constructor.
+			ThreadContext.IsInConstructor--;
+			check(ThreadContext.IsInConstructor >= 0);
+			ThreadContext.ConstructedObject = LastConstructedObject;
+		};
 
 		check(Obj != nullptr);
 
@@ -3920,9 +3958,15 @@ FObjectInitializer::~FObjectInitializer()
 	if (!bIsDeferredInitializer)
 #endif // USE_CIRCULAR_DEPENDENCY_LOAD_DEFERRING		
 	{
-
-		check(ThreadContext.TopInitializer() == this);
-		ThreadContext.PopInitializer();
+		// Always write to FUObjectThreadContext from the open, as writing to a 
+		// memory location in both the open and closed within the same transaction
+		// can lead to state corruption.
+		AutoRTFM::PopOnAbortHandler(this);
+		UE_AUTORTFM_OPEN
+		{
+			check(ThreadContext.TopInitializer() == this);
+			ThreadContext.PopInitializer();
+		};
 	}
 }
 
@@ -4092,10 +4136,18 @@ void FObjectInitializer::PostConstructInit()
 	Class->PostInitInstance(Obj, InstanceGraph);
 
 #if !(UE_BUILD_SHIPPING || UE_BUILD_TEST)
-	if (!FUObjectThreadContext::Get().PostInitPropertiesCheck.Num() || (FUObjectThreadContext::Get().PostInitPropertiesCheck.Pop(EAllowShrinking::No) != Obj))
+	FUObjectThreadContext& ThreadContext = FUObjectThreadContext::Get();
+	// Always write to FUObjectThreadContext from the open, as writing to a 
+	// memory location in both the open and closed within the same transaction
+	// can lead to state corruption.
+	AutoRTFM::PopOnAbortHandler(&ThreadContext.PostInitPropertiesCheck);
+	UE_AUTORTFM_OPEN
 	{
-		UE_LOG(LogUObjectGlobals, Fatal, TEXT("%s failed to route PostInitProperties. Call Super::PostInitProperties() in %s::PostInitProperties()."), *Obj->GetClass()->GetName(), *Obj->GetClass()->GetName());
-	}
+		if (!ThreadContext.PostInitPropertiesCheck.Num() || (ThreadContext.PostInitPropertiesCheck.Pop(EAllowShrinking::No) != Obj))
+		{
+			UE_LOG(LogUObjectGlobals, Fatal, TEXT("%s failed to route PostInitProperties. Call Super::PostInitProperties() in %s::PostInitProperties()."), *Obj->GetClass()->GetName(), *Obj->GetClass()->GetName());
+		}
+	};
 #endif // !(UE_BUILD_SHIPPING || UE_BUILD_TEST)
 
 #if !USE_CIRCULAR_DEPENDENCY_LOAD_DEFERRING
