[AutoRTFM] Migrate to UE::TScopeLock instead of FTransactionallySafeScopeLock.

These are functionally identical.

diff --git a/Engine/Source/Programs/AutoRTFMTests/Private/UECoreTests.cpp b/Engine/Source/Programs/AutoRTFMTests/Private/UECoreTests.cpp
--- a/Engine/Source/Programs/AutoRTFMTests/Private/UECoreTests.cpp
+++ b/Engine/Source/Programs/AutoRTFMTests/Private/UECoreTests.cpp
@@ -971,7 +971,7 @@ TEST_CASE("UECore.FUObjectItem")
 	}
 }
 
-TEST_CASE("UECore.TransactionallySafeScopeLock")
+TEST_CASE("UECore.TScopeLock_TransactionallySafeCriticalSection")
 {
 	SECTION("Outside Transaction")
 	{
@@ -979,7 +979,7 @@ TEST_CASE("UECore.TransactionallySafeScopeLock")
 
 		AutoRTFM::ETransactionResult Result = AutoRTFM::Transact([&]
 			{
-				FTransactionallySafeScopeLock Lock(&CriticalSection);
+				UE::TScopeLock Lock(CriticalSection);
 				AutoRTFM::AbortTransaction();
 			});
 
@@ -987,7 +987,7 @@ TEST_CASE("UECore.TransactionallySafeScopeLock")
 
 		Result = AutoRTFM::Transact([&]
 			{
-				FTransactionallySafeScopeLock Lock(&CriticalSection);
+				UE::TScopeLock Lock(CriticalSection);
 			});
 
 		REQUIRE(AutoRTFM::ETransactionResult::Committed == Result);
@@ -998,7 +998,7 @@ TEST_CASE("UECore.TransactionallySafeScopeLock")
 		AutoRTFM::ETransactionResult Result = AutoRTFM::Transact([&]
 			{
 				FTransactionallySafeCriticalSection CriticalSection;
-				FTransactionallySafeScopeLock Lock(&CriticalSection);
+				UE::TScopeLock Lock(CriticalSection);
 				AutoRTFM::AbortTransaction();
 			});
 
@@ -1007,7 +1007,7 @@ TEST_CASE("UECore.TransactionallySafeScopeLock")
 		Result = AutoRTFM::Transact([&]
 			{
 				FTransactionallySafeCriticalSection CriticalSection;
-				FTransactionallySafeScopeLock Lock(&CriticalSection);
+				UE::TScopeLock Lock(CriticalSection);
 			});
 
 		REQUIRE(AutoRTFM::ETransactionResult::Committed == Result);
@@ -1023,7 +1023,7 @@ TEST_CASE("UECore.TransactionallySafeScopeLock")
 
 				InnerResult = AutoRTFM::Transact([&]
 					{
-						FTransactionallySafeScopeLock Lock(&CriticalSection);
+						UE::TScopeLock Lock(CriticalSection);
 						AutoRTFM::CascadingAbortTransaction();
 					});
 			});
@@ -1036,7 +1036,7 @@ TEST_CASE("UECore.TransactionallySafeScopeLock")
 
 				InnerResult = AutoRTFM::Transact([&]
 					{
-						FTransactionallySafeScopeLock Lock(&CriticalSection);
+						UE::TScopeLock Lock(CriticalSection);
 						AutoRTFM::AbortTransaction();
 					});
 			});
@@ -1050,7 +1050,7 @@ TEST_CASE("UECore.TransactionallySafeScopeLock")
 
 				InnerResult = AutoRTFM::Transact([&]
 					{
-						FTransactionallySafeScopeLock Lock(&CriticalSection);
+						UE::TScopeLock Lock(CriticalSection);
 					});
 
 				AutoRTFM::AbortTransaction();
@@ -1064,7 +1064,7 @@ TEST_CASE("UECore.TransactionallySafeScopeLock")
 
 				InnerResult = AutoRTFM::Transact([&]
 					{
-						FTransactionallySafeScopeLock Lock(&CriticalSection);
+						UE::TScopeLock Lock(CriticalSection);
 					});
 			});
 
