[AutoRTFM] Wrap Catch2 macros to make them AutoRTFM-safe.

And fix the tests that were silently failing.

diff --git a/Engine/Source/Programs/AutoRTFMTests/Private/Catch2Includes.h b/Engine/Source/Programs/AutoRTFMTests/Private/Catch2Includes.h
--- a/Engine/Source/Programs/AutoRTFMTests/Private/Catch2Includes.h
+++ b/Engine/Source/Programs/AutoRTFMTests/Private/Catch2Includes.h
@@ -2,10 +2,91 @@
 
 #pragma once
 
-// we need to include thie MainCPP include *before* we include catch_amalgamated or we get issues with windows symbols
+// we need to include the MainCPP include *before* we include catch_amalgamated or we get issues with windows symbols
 #include "CoreMinimal.h"
 #include "LaunchEngineLoop.h"
+#include <AutoRTFM/AutoRTFM.h>
+
+#if defined(REQUIRE) || defined(CATCH_REQUIRE)
+#error "Catch2 already included"
+#endif
+
+// Prefixes all the 'real' Catch2 macros with 'CATCH_'.
+// We call these with the AutoRTFM-safe, unprefixed forms below.
+#define CATCH_CONFIG_PREFIX_ALL 1
 
 THIRD_PARTY_INCLUDES_START
 #include <catch_amalgamated.hpp>
 THIRD_PARTY_INCLUDES_END
+
+////////////////////////////////////////////////////////////////////////////////
+// AutoRTFM-safe versions of the Catch2 macros
+////////////////////////////////////////////////////////////////////////////////
+
+// Evaluates the expression (varargs) in the current transaction state, then
+// switches to the open to call into the Catch2 assertion handler logic.
+#define AUTORTFM_CATCH2_INTERNAL_CATCH_TEST(macroName, resultDisposition, ...) \
+    do { /* NOLINT(bugprone-infinite-loop) */ \
+        auto Expr = Catch::Decomposer() <= __VA_ARGS__; \
+        AutoRTFM::Open([&] \
+        { \
+            /* The expression should not be evaluated, but warnings should hopefully be checked */ \
+            CATCH_INTERNAL_IGNORE_BUT_WARN(__VA_ARGS__); \
+            Catch::AssertionHandler catchAssertionHandler( macroName##_catch_sr, CATCH_INTERNAL_LINEINFO, CATCH_INTERNAL_STRINGIFY(__VA_ARGS__), resultDisposition ); \
+            INTERNAL_CATCH_TRY { \
+                CATCH_INTERNAL_START_WARNINGS_SUPPRESSION \
+                CATCH_INTERNAL_SUPPRESS_PARENTHESES_WARNINGS \
+                catchAssertionHandler.handleExpr( Expr ); \
+                CATCH_INTERNAL_STOP_WARNINGS_SUPPRESSION \
+            } INTERNAL_CATCH_CATCH( catchAssertionHandler ) \
+            INTERNAL_CATCH_REACT( catchAssertionHandler ) \
+        }); \
+    } while( (void)0, (false) && static_cast<const bool&>( !!(__VA_ARGS__) ) )
+
+// Evaluates the matcher + arg in the current transaction state, then
+// switches to the open to call into the Catch2 assertion handler logic.
+#define AUTORTFM_CATCH2_INTERNAL_CHECK_THAT( macroName, matcher, resultDisposition, arg ) \
+    do { \
+        auto Expr = Catch::makeMatchExpr( arg, matcher ); \
+        AutoRTFM::Open([&] \
+        { \
+            Catch::AssertionHandler catchAssertionHandler( macroName##_catch_sr, CATCH_INTERNAL_LINEINFO, CATCH_INTERNAL_STRINGIFY(arg) ", " CATCH_INTERNAL_STRINGIFY(matcher), resultDisposition ); \
+            INTERNAL_CATCH_TRY { \
+                catchAssertionHandler.handleExpr( Expr ); \
+            } INTERNAL_CATCH_CATCH( catchAssertionHandler ) \
+            INTERNAL_CATCH_REACT( catchAssertionHandler ) \
+        }); \
+    } while(false)
+
+// Same as Catch2's INTERNAL_CATCH_SECTION, but the unique variable name is
+// passed in, and fails the test if called from a closed transaction.
+#define AUTORTFM_CATCH2_SECTION(VAR_NAME, ...)                              \
+        CATCH_INTERNAL_START_WARNINGS_SUPPRESSION                           \
+        CATCH_INTERNAL_SUPPRESS_UNUSED_VARIABLE_WARNINGS                    \
+        if ( Catch::Section const& VAR_NAME =                               \
+                Catch::Section( CATCH_INTERNAL_LINEINFO, __VA_ARGS__ );     \
+                AutoRTFMCatch2::EnsureNotClosed("SECTION") && VAR_NAME)     \
+        CATCH_INTERNAL_STOP_WARNINGS_SUPPRESSION
+
+#define REQUIRE( ... ) AUTORTFM_CATCH2_INTERNAL_CATCH_TEST( "REQUIRE", Catch::ResultDisposition::Normal, __VA_ARGS__ )
+#define REQUIRE_THAT( arg, matcher ) AUTORTFM_CATCH2_INTERNAL_CHECK_THAT( "REQUIRE_THAT", matcher, Catch::ResultDisposition::Normal, arg )
+#define FAIL(...) AutoRTFM::Open([&] { CATCH_FAIL(__VA_ARGS__); })
+#define TEST_CASE(...) CATCH_TEST_CASE(__VA_ARGS__)
+#define TEMPLATE_TEST_CASE(...) CATCH_TEMPLATE_TEST_CASE(__VA_ARGS__)
+#define BENCHMARK(...) CATCH_BENCHMARK(__VA_ARGS__)
+#define SECTION(...) AUTORTFM_CATCH2_SECTION(INTERNAL_CATCH_UNIQUE_NAME(catch_internal_Section), __VA_ARGS__)
+
+namespace AutoRTFMCatch2
+{
+
+inline bool EnsureNotClosed(const char* Where)
+{
+     if (AutoRTFM::IsClosed())
+     {
+        FAIL(Where << " cannot be used from a closed transaction");
+        return false;
+     }
+     return true;
+}
+
+} // AutoRTFMCatch2
diff --git a/Engine/Source/Programs/AutoRTFMTests/Private/OpenAPI.cpp b/Engine/Source/Programs/AutoRTFMTests/Private/OpenAPI.cpp
--- a/Engine/Source/Programs/AutoRTFMTests/Private/OpenAPI.cpp
+++ b/Engine/Source/Programs/AutoRTFMTests/Private/OpenAPI.cpp
@@ -106,7 +106,7 @@ TEST_CASE("OpenAPI.AbortTransactionScopedFromClosed")
 {
 	auto transactResult = AutoRTFM::Transact([&]()
 	{
-		REQUIRE(AutoRTFM::EContextStatus::OnTrack == AutoRTFM::Close([&]()
+		REQUIRE(AutoRTFM::EContextStatus::AbortedByRequest == AutoRTFM::Close([&]()
 		{
 			AutoRTFM::AbortTransaction();
 		}));
@@ -701,7 +701,7 @@ TEST_CASE("OpenAPI.Footgun1")
 	AutoRTFM::Transact([&]()
 	{
 		// Does nothing - already closed
-		REQUIRE(AutoRTFM::EContextStatus::OnTrack == AutoRTFM::Close([&]()
+		REQUIRE(AutoRTFM::EContextStatus::AbortedByRequest == AutoRTFM::Close([&]()
 		{
 			// Recorded valueB as starting at 0.
 			valueB = 123;
@@ -737,7 +737,7 @@ TEST_CASE("OpenAPI.Footgun2")
 	AutoRTFM::Transact([&]()
 	{
 		// Does nothing - already closed
-		REQUIRE(AutoRTFM::EContextStatus::OnTrack == AutoRTFM::Close([&]()
+		REQUIRE(AutoRTFM::EContextStatus::AbortedByRequest == AutoRTFM::Close([&]()
 		{
 			// Recorded valueB as starting at 0.
 			valueB = 20;
@@ -868,7 +868,7 @@ TEST_CASE("OpenAPI.TransOpenTransCloseAbortAbort")
 
 				value = 42;
 				// Can't close outside of a Transact
-				REQUIRE(AutoRTFM::EContextStatus::OnTrack == AutoRTFM::Close([&]()
+				REQUIRE(AutoRTFM::EContextStatus::AbortedByRequest == AutoRTFM::Close([&]()
 				{
 					value = 420;
 					AutoRTFM::AbortTransaction();
