Fix the Verse tests build breakage.

#rnx

diff --git a/Engine/Source/Programs/AutoRTFMTests/Private/UECoreTests.cpp b/Engine/Source/Programs/AutoRTFMTests/Private/UECoreTests.cpp
--- a/Engine/Source/Programs/AutoRTFMTests/Private/UECoreTests.cpp
+++ b/Engine/Source/Programs/AutoRTFMTests/Private/UECoreTests.cpp
@@ -12,7 +12,6 @@
 #include "UObject/NameTypes.h"
 #include "UObject/UObjectArray.h"
 #include "MyAutoRTFMTestObject.h"
-#include "Misc/TransactionallySafeScopeLock.h"
 #include "Misc/TransactionallySafeRWScopeLock.h"
 #include "Misc/TransactionallySafeMutex.h"
 #include "Containers/Queue.h"
