[AutoRTFM] Clean up AutoRTFMTests now that Catch2 is AutoRTFM-safe

We now don't have to jump through so many hoops to make these tests work reliably.
Clean up the test logic so it's a bit more readable.

diff --git a/Engine/Source/Programs/AutoRTFMTests/Private/API.cpp b/Engine/Source/Programs/AutoRTFMTests/Private/API.cpp
--- a/Engine/Source/Programs/AutoRTFMTests/Private/API.cpp
+++ b/Engine/Source/Programs/AutoRTFMTests/Private/API.cpp
@@ -62,13 +62,12 @@ TEST_CASE("API.autortfm_abort_transaction")
     bool BeforeNest = false;
     bool InNest = false;
     bool AfterNest = false;
-    AutoRTFM::ETransactionResult NestResult;
-
+    
     AutoRTFM::Commit([&]
     {
         BeforeNest = true;
 
-        NestResult = AutoRTFM::Transact([&]
+        AutoRTFM::ETransactionResult Result = AutoRTFM::Transact([&]
         {
             // Because we are aborting this won't actually occur!
             InNest = true;
@@ -76,13 +75,14 @@ TEST_CASE("API.autortfm_abort_transaction")
             autortfm_abort_transaction();
         });
 
+        REQUIRE(AutoRTFM::ETransactionResult::AbortedByRequest == Result);
+
         AfterNest = true;
     });
 
     REQUIRE(true == BeforeNest);
     REQUIRE(false == InNest);
     REQUIRE(true == AfterNest);
-    REQUIRE(AutoRTFM::ETransactionResult::AbortedByRequest == NestResult);
 }
 
 TEST_CASE("API.autortfm_abort_if_transactional")
@@ -93,13 +93,12 @@ TEST_CASE("API.autortfm_abort_if_transactional")
     bool BeforeNest = false;
     bool InNest = false;
     bool AfterNest = false;
-    AutoRTFM::ETransactionResult NestResult;
-
+    
     AutoRTFM::Commit([&]
     {
         BeforeNest = true;
 
-        NestResult = AutoRTFM::Transact([&]
+        AutoRTFM::ETransactionResult Result = AutoRTFM::Transact([&]
         {
             // Because we are aborting this won't actually occur!
             InNest = true;
@@ -107,13 +106,14 @@ TEST_CASE("API.autortfm_abort_if_transactional")
             autortfm_abort_if_transactional();
         });
 
+        REQUIRE(AutoRTFM::ETransactionResult::AbortedByRequest == Result);
+
         AfterNest = true;
     });
 
     REQUIRE(true == BeforeNest);
     REQUIRE(false == InNest);
     REQUIRE(true == AfterNest);
-    REQUIRE(AutoRTFM::ETransactionResult::AbortedByRequest == NestResult);
 }
 
 TEST_CASE("API.autortfm_abort_if_closed")
@@ -195,7 +195,6 @@ TEST_CASE("API.autortfm_on_commit")
     bool InnerTransaction = false;
     bool InnerTransactionWithAbort = false;
     bool InnerOpenNest = false;
-    AutoRTFM::ETransactionResult NestResult;
 
     AutoRTFM::Commit([&]
     {
@@ -224,7 +223,7 @@ TEST_CASE("API.autortfm_on_commit")
             AutoRTFM::AbortTransaction();
         }
 
-        NestResult = AutoRTFM::Transact([&]
+        AutoRTFM::ETransactionResult Result = AutoRTFM::Transact([&]
         {
             autortfm_on_commit([](void* const Arg)
             {
@@ -234,6 +233,8 @@ TEST_CASE("API.autortfm_on_commit")
             AutoRTFM::AbortTransaction();
         });
 
+        REQUIRE(AutoRTFM::ETransactionResult::AbortedByRequest == Result);
+
         // This should never be modified because its transaction aborted!
         if (InnerTransactionWithAbort)
         {
@@ -259,7 +260,6 @@ TEST_CASE("API.autortfm_on_commit")
     REQUIRE(true == InnerTransaction);
     REQUIRE(false == InnerTransactionWithAbort);
     REQUIRE(true == InnerOpenNest);
-    REQUIRE(AutoRTFM::ETransactionResult::AbortedByRequest == NestResult);
 }
 
 TEST_CASE("API.autortfm_on_abort")
@@ -274,7 +274,6 @@ TEST_CASE("API.autortfm_on_abort")
     bool InnerTransaction = false;
     bool InnerTransactionWithAbort = false;
     bool InnerOpenNest = false;
-    AutoRTFM::ETransactionResult NestResult = AutoRTFM::ETransactionResult::Committed;
 
     REQUIRE(AutoRTFM::ETransactionResult::Committed == AutoRTFM::Transact([&]
     {
@@ -285,7 +284,6 @@ TEST_CASE("API.autortfm_on_abort")
 			InnerTransaction = false;
 			InnerTransactionWithAbort = false;
 			InnerOpenNest = false;
-			NestResult = AutoRTFM::ETransactionResult::Committed;
 		});
 
         autortfm_on_abort([](void* const Arg)
@@ -313,7 +311,7 @@ TEST_CASE("API.autortfm_on_abort")
             AutoRTFM::AbortTransaction();
         }
 
-        NestResult = AutoRTFM::Transact([&]
+        AutoRTFM::ETransactionResult Result = AutoRTFM::Transact([&]
         {
             autortfm_on_abort([](void* const Arg)
             {
@@ -323,6 +321,8 @@ TEST_CASE("API.autortfm_on_abort")
             AutoRTFM::AbortTransaction();
         });
 
+        REQUIRE(AutoRTFM::ETransactionResult::AbortedByRequest == Result);
+
         // OnAbort runs eagerly on inner abort
         if (!InnerTransactionWithAbort)
         {
@@ -348,7 +348,6 @@ TEST_CASE("API.autortfm_on_abort")
     REQUIRE(false == InnerTransaction);
     REQUIRE(true == InnerTransactionWithAbort);
     REQUIRE(false == InnerOpenNest);
-    REQUIRE(AutoRTFM::ETransactionResult::AbortedByRequest == NestResult);
 }
 
 TEST_CASE("API.autortfm_did_allocate")
@@ -613,13 +612,12 @@ TEST_CASE("API.Abort")
     bool BeforeNest = false;
     bool InNest = false;
     bool AfterNest = false;
-    AutoRTFM::ETransactionResult NestResult;
 
     AutoRTFM::Commit([&]
     {
         BeforeNest = true;
 
-        NestResult = AutoRTFM::Transact([&]
+        AutoRTFM::ETransactionResult Result = AutoRTFM::Transact([&]
         {
             // Because we are aborting this won't actually occur!
             InNest = true;
@@ -627,13 +625,14 @@ TEST_CASE("API.Abort")
             AutoRTFM::AbortTransaction();
         });
 
+        REQUIRE(AutoRTFM::ETransactionResult::AbortedByRequest == Result);
+
         AfterNest = true;
     });
 
     REQUIRE(true == BeforeNest);
     REQUIRE(false == InNest);
     REQUIRE(true == AfterNest);
-    REQUIRE(AutoRTFM::ETransactionResult::AbortedByRequest == NestResult);
 }
 
 TEST_CASE("API.AbortIfTransactional")
@@ -644,13 +643,12 @@ TEST_CASE("API.AbortIfTransactional")
     bool BeforeNest = false;
     bool InNest = false;
     bool AfterNest = false;
-    AutoRTFM::ETransactionResult NestResult;
-
+    
     AutoRTFM::Commit([&]
     {
         BeforeNest = true;
 
-        NestResult = AutoRTFM::Transact([&]
+        AutoRTFM::ETransactionResult Result = AutoRTFM::Transact([&]
         {
             // Because we are aborting this won't actually occur!
             InNest = true;
@@ -658,13 +656,14 @@ TEST_CASE("API.AbortIfTransactional")
             AutoRTFM::AbortIfTransactional();
         });
 
+        REQUIRE(AutoRTFM::ETransactionResult::AbortedByRequest == Result);
+
         AfterNest = true;
     });
 
     REQUIRE(true == BeforeNest);
     REQUIRE(false == InNest);
     REQUIRE(true == AfterNest);
-    REQUIRE(AutoRTFM::ETransactionResult::AbortedByRequest == NestResult);
 }
 
 TEST_CASE("API.AbortIfClosed")
@@ -767,18 +766,17 @@ TEST_CASE("API.Close")
         AutoRTFM::EContextStatus CloseStatusA = AutoRTFM::Close([&] { InClosedNest = true; });
         REQUIRE(AutoRTFM::EContextStatus::OnTrack == CloseStatusA);
 
-        AutoRTFM::EContextStatus CloseStatusB{AutoRTFM::EContextStatus::Idle};
         AutoRTFM::Open([&]
         {
-            CloseStatusB = AutoRTFM::Close([&] { InClosedNestInOpenNest = true; });
+            AutoRTFM::EContextStatus CloseStatusB = AutoRTFM::Close([&]
+            {
+                InClosedNestInOpenNest = true;
+            });
+            REQUIRE(AutoRTFM::EContextStatus::OnTrack == CloseStatusB);
+
             InOpenNest = true;
         });
 
-        // A closed call inside an open does not abort either.
-        // Note: REQUIRE() must be called outside of the Open() lambda as the Catch2 framework may
-        // perform writes which are being tracked by the closed transaction.
-        REQUIRE(AutoRTFM::EContextStatus::OnTrack == CloseStatusB);
-
         AutoRTFM::AbortTransaction();
     });
 
@@ -811,7 +809,6 @@ TEST_CASE("API.OnCommit")
     bool InnerTransaction = false;
     bool InnerTransactionWithAbort = false;
     bool InnerOpenNest = false;
-    AutoRTFM::ETransactionResult NestResult;
 
     REQUIRE(AutoRTFM::ETransactionResult::Committed == AutoRTFM::Transact([&]
     {
@@ -840,16 +837,19 @@ TEST_CASE("API.OnCommit")
 			AutoRTFM::AbortTransaction();
         }
 
-        NestResult = AutoRTFM::Transact([&]
+
+        AutoRTFM::ETransactionResult Result = AutoRTFM::Transact([&]
         {
             AutoRTFM::OnCommit([&]
             {
                 InnerTransactionWithAbort = true;
             });
 
-		AutoRTFM::AbortTransaction();
+            AutoRTFM::AbortTransaction();
         });
 
+        REQUIRE(AutoRTFM::ETransactionResult::AbortedByRequest == Result);
+
         // This should never be modified because its transaction aborted!
         if (InnerTransactionWithAbort)
         {
@@ -875,7 +875,6 @@ TEST_CASE("API.OnCommit")
     REQUIRE(true == InnerTransaction);
     REQUIRE(false == InnerTransactionWithAbort);
     REQUIRE(true == InnerOpenNest);
-    REQUIRE(AutoRTFM::ETransactionResult::AbortedByRequest == NestResult);
 }
 
 TEST_CASE("API.OnCommit_MutableCapture")
@@ -940,7 +939,6 @@ TEST_CASE("API.OnAbort")
     bool InnerTransaction = false;
     bool InnerTransactionWithAbort = false;
     bool InnerOpenNest = false;
-    AutoRTFM::ETransactionResult NestResult = AutoRTFM::ETransactionResult::Committed;
 
     REQUIRE(AutoRTFM::ETransactionResult::Committed == AutoRTFM::Transact([&]
     {
@@ -951,7 +949,6 @@ TEST_CASE("API.OnAbort")
 			InnerTransaction = false;
 			InnerTransactionWithAbort = false;
 			InnerOpenNest = false;
-			NestResult = AutoRTFM::ETransactionResult::Committed;
 		});
 
         AutoRTFM::OnAbort([&]
@@ -979,7 +976,7 @@ TEST_CASE("API.OnAbort")
 			AutoRTFM::AbortTransaction();
         }
 
-        NestResult = AutoRTFM::Transact([&]
+        AutoRTFM::ETransactionResult Result = AutoRTFM::Transact([&]
         {
             AutoRTFM::OnAbort([&]
             {
@@ -989,6 +986,8 @@ TEST_CASE("API.OnAbort")
 			AutoRTFM::AbortTransaction();
         });
 
+        REQUIRE(AutoRTFM::ETransactionResult::AbortedByRequest == Result);
+
         // Inner OnAbort runs eagerly
         if (!InnerTransactionWithAbort)
         {
@@ -1014,7 +1013,6 @@ TEST_CASE("API.OnAbort")
     REQUIRE(false == InnerTransaction);
     REQUIRE(true == InnerTransactionWithAbort);
     REQUIRE(false == InnerOpenNest);
-    REQUIRE(AutoRTFM::ETransactionResult::AbortedByRequest == NestResult);
 }
 
 TEST_CASE("API.OnAbort_MutableCapture")
diff --git a/Engine/Source/Programs/AutoRTFMTests/Private/Abort.cpp b/Engine/Source/Programs/AutoRTFMTests/Private/Abort.cpp
--- a/Engine/Source/Programs/AutoRTFMTests/Private/Abort.cpp
+++ b/Engine/Source/Programs/AutoRTFMTests/Private/Abort.cpp
@@ -50,7 +50,6 @@ TEST_CASE("Abort")
 
 TEST_CASE("Abort.NestedAbortOrder")
 {
-	AutoRTFM::ETransactionResult InnerResult;
 	unsigned Orderer = 0;
 
 	AutoRTFM::Commit([&]
@@ -61,7 +60,7 @@ TEST_CASE("Abort.NestedAbortOrder")
 			Orderer = 0;
 		});
 
-		InnerResult = AutoRTFM::Transact([&]
+		AutoRTFM::ETransactionResult Result = AutoRTFM::Transact([&]
 		{
 			AutoRTFM::OnAbort([&]
 				{
@@ -77,57 +76,52 @@ TEST_CASE("Abort.NestedAbortOrder")
 
 			AutoRTFM::AbortTransaction();
 		});
+		REQUIRE(AutoRTFM::ETransactionResult::AbortedByRequest == Result);
 	});
 
-	REQUIRE(AutoRTFM::ETransactionResult::AbortedByRequest == InnerResult);
 	REQUIRE(2 == Orderer);
 }
 
 TEST_CASE("Abort.TransactionInOnCommit")
 {
-	AutoRTFM::ETransactionResult InnerResult;
-
 	AutoRTFM::Commit([&]
 	{
 		AutoRTFM::OnCommit([&]
 		{
 			bool bDidSomething = false;
 
-			InnerResult = AutoRTFM::Transact([&]
+			AutoRTFM::ETransactionResult Result = AutoRTFM::Transact([&]
 			{
 				bDidSomething = true;
 			});
 
+			REQUIRE(AutoRTFM::ETransactionResult::AbortedByTransactInOnCommit == Result);
 			REQUIRE(false == bDidSomething);
 		});
 	});
 
-	REQUIRE(AutoRTFM::ETransactionResult::AbortedByTransactInOnCommit == InnerResult);
 }
 
 TEST_CASE("Abort.TransactionInOnAbort")
 {
-	AutoRTFM::ETransactionResult Result;
-	AutoRTFM::ETransactionResult InnerResult;
-
-	Result = AutoRTFM::Transact([&]
+	AutoRTFM::ETransactionResult Result = AutoRTFM::Transact([&]
 	{
 		AutoRTFM::OnAbort([&]
 		{
 			bool bDidSomething = false;
 
-			InnerResult = AutoRTFM::Transact([&]
+			AutoRTFM::ETransactionResult InnerResult = AutoRTFM::Transact([&]
 			{
 				bDidSomething = true;
 			});
 
+			REQUIRE(AutoRTFM::ETransactionResult::AbortedByTransactInOnAbort == InnerResult);
 			REQUIRE(false == bDidSomething);
 		});
 
 		AutoRTFM::AbortTransaction();
 	});
 
-	REQUIRE(AutoRTFM::ETransactionResult::AbortedByTransactInOnAbort == InnerResult);
 	REQUIRE(AutoRTFM::ETransactionResult::AbortedByRequest == Result);
 }
 
@@ -679,9 +673,8 @@ TEST_CASE("Abort.LanguageThroughOpen")
 // Test for SOL-5804
 TEST_CASE("Abort.StackWriteToOuterOpen")
 {
-	std::string_view TestResult;
-	AutoRTFM::EContextStatus CloseStatus = AutoRTFM::EContextStatus::Idle;
 	bool WritesUndone = true;
+	bool Success = false;
 
 	const AutoRTFM::ETransactionResult TransactionResult = AutoRTFM::Transact([&]
 	{
@@ -689,7 +682,7 @@ TEST_CASE("Abort.StackWriteToOuterOpen")
 		{
 			std::array<int, 64> Values{};
 
-			CloseStatus = AutoRTFM::Close([&]
+			AutoRTFM::EContextStatus Status = AutoRTFM::Close([&]
 			{
 				// On stack outside transaction.
 				// Should be reverted as part of the abort.
@@ -702,6 +695,8 @@ TEST_CASE("Abort.StackWriteToOuterOpen")
 					Values[I] = static_cast<int>(I * 10);
 				}
 			});
+
+			REQUIRE(AutoRTFM::EContextStatus::OnTrack == Status);
 		});
 
 		// If any of the variables on the stack within the Open() get written to
@@ -713,15 +708,15 @@ TEST_CASE("Abort.StackWriteToOuterOpen")
 		{
 			if (!WritesUndone)
 			{
-				TestResult = "OnAbort was called without first reverting memory";
+				FAIL("OnAbort was called without first reverting memory");
 			}
 			else if (StackGuard != std::array<int, 64>{})
 			{
-				TestResult = "StackGuard was corrupted";
+				FAIL("StackGuard was corrupted");
 			}
 			else
 			{
-				TestResult = "Success";
+				Success = true;
 			}
 		});
 
@@ -729,9 +724,8 @@ TEST_CASE("Abort.StackWriteToOuterOpen")
 		AutoRTFM::AbortTransaction();
 	});
 
+	REQUIRE(Success);
 	REQUIRE(AutoRTFM::ETransactionResult::AbortedByRequest == TransactionResult);
-	REQUIRE(AutoRTFM::EContextStatus::OnTrack == CloseStatus);
-	REQUIRE("Success" == TestResult);
 }
 
 TEST_CASE("Abort.PushOnAbortOrdering")
diff --git a/Engine/Source/Programs/AutoRTFMTests/Private/FunctionPointerTests.cpp b/Engine/Source/Programs/AutoRTFMTests/Private/FunctionPointerTests.cpp
--- a/Engine/Source/Programs/AutoRTFMTests/Private/FunctionPointerTests.cpp
+++ b/Engine/Source/Programs/AutoRTFMTests/Private/FunctionPointerTests.cpp
@@ -40,41 +40,32 @@ namespace
 
 TEST_CASE("FunctionPointer.CStyle")
 {
-	int Result = 0;
 	AutoRTFM::Commit([&]
 		{
 			CStyleType CStyle = GetTestCFunction();
-			Result = CStyle();
+			REQUIRE(42 == CStyle());
 		});
-
-	REQUIRE(42 == Result);
 }
 
 TEST_CASE("FunctionPointer.StandardLibrary")
 {
 	SECTION("Created inside transaction")
 	{
-		float Result = 0;
 		AutoRTFM::Commit([&]
 			{
 				CosfType fn = GetCosfFunction();
-				Result = fn(0.0f);
+				REQUIRE(1.0f == fn(0.0f));
 			});
-
-		REQUIRE(Result == 1.0f);
 	}
 
 	SECTION("Created outside transaction")
 	{
-		int Result = 0;
 		CosfType fn = GetCosfFunction();
 
 		AutoRTFM::Commit([&]
 			{
-				Result = fn(0.0f);
+				REQUIRE(1.0f == fn(0.0f));
 			});
-
-		REQUIRE(Result == 1.0f);
 	}
 }
 
diff --git a/Engine/Source/Programs/AutoRTFMTests/Private/Open.cpp b/Engine/Source/Programs/AutoRTFMTests/Private/Open.cpp
--- a/Engine/Source/Programs/AutoRTFMTests/Private/Open.cpp
+++ b/Engine/Source/Programs/AutoRTFMTests/Private/Open.cpp
@@ -185,14 +185,12 @@ TEST_CASE("Open.ReturnValue")
             };
             static_assert(AutoRTFM::IsSafeToReturnFromOpen<FMyStruct>);
 
-            FMyStruct StructResult;
             AutoRTFM::ETransactionResult TransactionResult = AutoRTFM::Transact([&]
             {
                 FMyStruct Closed = AutoRTFM::Open([] { return FMyStruct{42}; });
-                AutoRTFM::Open([&] { StructResult = Closed; });
+                REQUIRE(42 == Closed.Value);
             });
             REQUIRE(AutoRTFM::ETransactionResult::Committed == TransactionResult);
-            REQUIRE(42 == StructResult.Value);
         }
         SECTION("AutoRTFMAssignFromOpenToClosed() by const-ref")
         {
@@ -206,14 +204,12 @@ TEST_CASE("Open.ReturnValue")
             };
             static_assert(AutoRTFM::IsSafeToReturnFromOpen<FMyStruct>);
 
-            FMyStruct StructResult;
             AutoRTFM::ETransactionResult TransactionResult = AutoRTFM::Transact([&]
             {
                 FMyStruct Closed = AutoRTFM::Open([] { return FMyStruct{42}; });
-                AutoRTFM::Open([&] { StructResult = Closed; });
+                REQUIRE(42 == Closed.Value);
             });
             REQUIRE(AutoRTFM::ETransactionResult::Committed == TransactionResult);
-            REQUIRE(42 == StructResult.Value);
         }
         SECTION("AutoRTFMAssignFromOpenToClosed() by rvalue-ref")
         {
@@ -229,16 +225,14 @@ TEST_CASE("Open.ReturnValue")
             };
             static_assert(AutoRTFM::IsSafeToReturnFromOpen<FMyStruct>);
 
-            bool WasMoved = false;
-            FMyStruct StructResult;
             AutoRTFM::ETransactionResult TransactionResult = AutoRTFM::Transact([&]
             {
+                bool WasMoved = false;
                 FMyStruct Closed = AutoRTFM::Open([&] { return FMyStruct{42, &WasMoved}; });
-                AutoRTFM::Open([&] { StructResult = Closed; });
+                REQUIRE(true == WasMoved);
+                REQUIRE(42 == Closed.Value);
             });
             REQUIRE(AutoRTFM::ETransactionResult::Committed == TransactionResult);
-            REQUIRE(true == WasMoved);
-            REQUIRE(42 == StructResult.Value);
         }
     }
 }
diff --git a/Engine/Source/Programs/AutoRTFMTests/Private/OpenAPI.cpp b/Engine/Source/Programs/AutoRTFMTests/Private/OpenAPI.cpp
--- a/Engine/Source/Programs/AutoRTFMTests/Private/OpenAPI.cpp
+++ b/Engine/Source/Programs/AutoRTFMTests/Private/OpenAPI.cpp
@@ -485,8 +485,6 @@ TEST_CASE("OpenAPI.StackWriteCommitInTheOpen1")
 
 TEST_CASE("OpenAPI.StackWriteCommitInTheOpen2")
 {
-	AutoRTFM::EContextStatus Status = AutoRTFM::EContextStatus::Idle;
-
 	int Value = 0;
 
 	AutoRTFM::Transact([&]()
@@ -495,7 +493,7 @@ TEST_CASE("OpenAPI.StackWriteCommitInTheOpen2")
 		{
 			AutoRTFM::ForTheRuntime::StartTransaction();
 
-			Status = AutoRTFM::Close([&]()
+			AutoRTFM::EContextStatus Status = AutoRTFM::Close([&]()
 			{
 				AutoRTFM::Open([&]()
 				{
@@ -503,13 +501,11 @@ TEST_CASE("OpenAPI.StackWriteCommitInTheOpen2")
 					Value = 10;
 				});
 			});
+			REQUIRE(AutoRTFM::EContextStatus::OnTrack == Status);
 
 			AutoRTFM::ForTheRuntime::CommitTransaction();
 		});
 
-		// Note: REQUIRE() must be called outside of the Open() lambda as the Catch2 framework may
-		// perfore writes which are being tracked by the closed transaction.
-		REQUIRE(AutoRTFM::EContextStatus::OnTrack == Status);
 		REQUIRE(Value == 10);
 	});
 }
diff --git a/Engine/Source/Programs/AutoRTFMTests/Private/UObjectTests.cpp b/Engine/Source/Programs/AutoRTFMTests/Private/UObjectTests.cpp
--- a/Engine/Source/Programs/AutoRTFMTests/Private/UObjectTests.cpp
+++ b/Engine/Source/Programs/AutoRTFMTests/Private/UObjectTests.cpp
@@ -184,24 +184,14 @@ TEST_CASE("FUObjectAnnotationSparse.AddAnnotation")
 	{
 		SECTION("Commit")
 		{
-			FTestAnnotation PreGet;
-			FTestAnnotation PostGet;
-
 			AutoRTFM::Transact([&]
 			{
-				FTestAnnotation Pre = AnnotationMap.GetAnnotation(Key);
+				REQUIRE(FTestAnnotation() == AnnotationMap.GetAnnotation(Key));
 				AnnotationMap.AddAnnotation(Key, ValueA);
-				FTestAnnotation Post = AnnotationMap.GetAnnotation(Key);
-				AutoRTFM::Open([&]
-				{
-					PreGet = Pre;
-					PostGet = Post;
-				});
+				REQUIRE(ValueA == AnnotationMap.GetAnnotation(Key));
 			});
 
 			REQUIRE(ValueA == AnnotationMap.GetAnnotation(Key));
-			REQUIRE(FTestAnnotation() == PreGet);
-			REQUIRE(ValueA == PostGet);
 		}
 
 		SECTION("Abort")
@@ -222,24 +212,14 @@ TEST_CASE("FUObjectAnnotationSparse.AddAnnotation")
 
 		SECTION("Commit")
 		{
-			FTestAnnotation PreGet;
-			FTestAnnotation PostGet;
-
 			AutoRTFM::Transact([&]
 			{
-				FTestAnnotation Pre = AnnotationMap.GetAnnotation(Key);
+				REQUIRE(ValueB == AnnotationMap.GetAnnotation(Key));
 				AnnotationMap.AddAnnotation(Key, ValueA);
-				FTestAnnotation Post = AnnotationMap.GetAnnotation(Key);
-				AutoRTFM::Open([&]
-				{
-					PreGet = Pre;
-					PostGet = Post;
-				});
+				REQUIRE(ValueA == AnnotationMap.GetAnnotation(Key));
 			});
 
 			REQUIRE(ValueA == AnnotationMap.GetAnnotation(Key));
-			REQUIRE(ValueB == PreGet);
-			REQUIRE(ValueA == PostGet);
 		}
 
 		SECTION("Abort")
@@ -364,12 +344,13 @@ TEST_CASE("UObject.FUObjectThreadContext")
 		
 		static void CtorCreateInnerClosed(const FObjectInitializer& ObjectInitializer, UMyAutoRTFMTestObject& Object)
 		{
-			(void) AutoRTFM::Close([&]
+			AutoRTFM::EContextStatus Status = AutoRTFM::Close([&]
 			{
 				REQUIRE(1 == FUObjectThreadContext::Get().IsInConstructor);
 				CreateObjectWithCtor(nullptr);
 				REQUIRE(1 == FUObjectThreadContext::Get().IsInConstructor);
 			});
+			REQUIRE(AutoRTFM::EContextStatus::OnTrack == Status);
 		}
 		
 		static void CtorCreateInnerTransact(const FObjectInitializer& ObjectInitializer, UMyAutoRTFMTestObject& Object)
@@ -385,11 +366,12 @@ TEST_CASE("UObject.FUObjectThreadContext")
 		static void CtorAbort(const FObjectInitializer& ObjectInitializer, UMyAutoRTFMTestObject& Object)
 		{
 			REQUIRE(1 == FUObjectThreadContext::Get().IsInConstructor);
-			(void) AutoRTFM::Close([&]
+			AutoRTFM::EContextStatus Status = AutoRTFM::Close([&]
 			{
 				AutoRTFM::AbortTransaction();
 				FAIL(/* unreachable */);
 			});
+			REQUIRE(AutoRTFM::EContextStatus::AbortedByRequest == Status);
 		}
 	};
 
