Add new string test

diff --git a/Engine/Source/Programs/AutoRTFMTests/Private/CString.cpp b/Engine/Source/Programs/AutoRTFMTests/Private/CString.cpp
--- a/Engine/Source/Programs/AutoRTFMTests/Private/CString.cpp
+++ b/Engine/Source/Programs/AutoRTFMTests/Private/CString.cpp
@@ -971,6 +971,38 @@ TEST_CASE("CString.wcsstr")
 	}
 }
 
+TEST_CASE("CString.wcsrchr")
+{
+	const wchar_t* A = L"This";
+	SECTION("With Abort")
+	{
+		const wchar_t* Return = nullptr;
+
+		CLANG_BEGIN_DISABLE_OPTIMIZATIONS
+		AutoRTFM::ETransactionResult Result = AutoRTFM::Transact([&]
+		{
+			Return = wcsrchr(A, 'h');
+			AutoRTFM::AbortTransaction();
+		});
+		CLANG_END_DISABLE_OPTIMIZATIONS
+
+		REQUIRE(AutoRTFM::ETransactionResult::AbortedByRequest == Result);
+		REQUIRE(nullptr == Return);
+	}
+
+	SECTION("With Commit")
+	{
+		const wchar_t* Return = nullptr;
+
+		CLANG_BEGIN_DISABLE_OPTIMIZATIONS
+		AutoRTFM::Commit([&] { Return = wcsrchr(A, 'h'); });
+		CLANG_END_DISABLE_OPTIMIZATIONS
+
+		REQUIRE(nullptr != Return);
+		REQUIRE('h' == *Return);
+	}
+}
+
 TEST_CASE("CString.wcscmp")
 {
 	const wchar_t* A = L"This";
