[AutoRTFM] Fix stack-use-after-free in MTTransactionallySafeAccessDetector tests

This can lead to memory corruption, and cause deadlocks on mutexes.

diff --git a/Engine/Source/Programs/AutoRTFMTests/Private/MTTSAccessDetector.cpp b/Engine/Source/Programs/AutoRTFMTests/Private/MTTSAccessDetector.cpp
--- a/Engine/Source/Programs/AutoRTFMTests/Private/MTTSAccessDetector.cpp
+++ b/Engine/Source/Programs/AutoRTFMTests/Private/MTTSAccessDetector.cpp
@@ -439,6 +439,7 @@ TEST_CASE("MTTransactionallySafeAccessDetector")
 	{
 		FRWTransactionallySafeAccessDetector Detector;
 		std::atomic_uint Handshake = 0; // Used to order this thread and the spawnee.
+		std::thread OtherThread;
 		AutoRTFM::ETransactionResult Result = AutoRTFM::Transact([&]
 			{
 				// Because we can retry transactions reset handshake to zero.
@@ -452,16 +453,18 @@ TEST_CASE("MTTransactionallySafeAccessDetector")
 
 					AutoRTFM::Open([&]
 						{
-							std::thread([&]
+							OtherThread = std::thread([&]
 								{
 									// Let the main thread progress.
 									Handshake += 1;
 
-									UE_MT_SCOPED_READ_ACCESS(Detector);
+									{
+										UE_MT_SCOPED_READ_ACCESS(Detector);
 
-									// Wait for the main thread.
-									while (2 != Handshake) {}
-								}).detach();
+										// Wait for the main thread.
+										while (2 != Handshake) {}
+									}
+								});
 						});
 				}
 
@@ -472,6 +475,11 @@ TEST_CASE("MTTransactionallySafeAccessDetector")
 
 						// Let the spawnee progress.
 						Handshake += 1;
+
+						// Wait for the spawnee to finish before unwinding the
+						// stack and invalidating the Handshake and Detector
+						// references used by the thread.
+						OtherThread.join();
 					});
 			});
 		REQUIRE(AutoRTFM::ETransactionResult::Committed == Result);
