[AutoRTFM] Test that AutoRTFM #includes its dependencies.

We were missing a required #include of Function.h because `Catch2Includes.h` transitively brings it in, so all our tests compiled. We now have a test cpp which #includes nothing but AutoRTFM.h and performs basic ONCOMMIT/ONABORT functionality. It does not register itself as a real test, but that doesn't matter; if it compiles and links, the test is successful.

diff --git a/Engine/Source/Programs/AutoRTFMTests/Private/DependencyTests.cpp b/Engine/Source/Programs/AutoRTFMTests/Private/DependencyTests.cpp
new file mode 100644
--- /dev/null
+++ b/Engine/Source/Programs/AutoRTFMTests/Private/DependencyTests.cpp
@@ -0,0 +1,52 @@
+// Copyright Epic Games, Inc. All Rights Reserved.
+
+#include <AutoRTFM/AutoRTFM.h>
+
+// We can't use Catch2Includes.h here, because it #includes more Unreal headers.
+// Fortunately, we don't even need to run these tests; if this code compiles and links, 
+// the test has passed.
+
+namespace 
+{
+
+void OnAbort_MacroWorksWithNoAdditionalIncludes()
+{
+	AutoRTFM::Transact([&]
+	{
+		UE_AUTORTFM_ONABORT()
+		{
+		};
+	});
+}
+
+void OnAbort_FunctionWorksWithNoAdditionalIncludes()
+{
+	AutoRTFM::Transact([&]
+	{
+		AutoRTFM::OnAbort([]
+		{
+		});
+	});
+}
+
+void OnCommit_MacroWorksWithNoAdditionalIncludes()
+{
+	AutoRTFM::Commit([&]
+	{
+		UE_AUTORTFM_ONCOMMIT()
+		{
+		};
+	});
+}
+
+void OnCommit_FunctionWorksWithNoAdditionalIncludes()
+{
+	AutoRTFM::Commit([&]
+	{
+		AutoRTFM::OnCommit([]
+		{
+		});
+	});
+}
+
+}  // namespace
diff --git a/Engine/Source/Runtime/Core/Public/AutoRTFM/AutoRTFM.h b/Engine/Source/Runtime/Core/Public/AutoRTFM/AutoRTFM.h
--- a/Engine/Source/Runtime/Core/Public/AutoRTFM/AutoRTFM.h
+++ b/Engine/Source/Runtime/Core/Public/AutoRTFM/AutoRTFM.h
@@ -64,9 +64,11 @@
 #define PRAGMA_DISABLE_UNREACHABLE_CODE_WARNINGS
 #define PRAGMA_RESTORE_UNREACHABLE_CODE_WARNINGS
 #define LIKELY(EXPR) __builtin_expect(!!(EXPR), 1)
+template <typename FuncType> class TFunction;  // Should be replaced with `std::function`
 #else
 #include <HAL/Platform.h>
 #include <HAL/PlatformMemory.h>
+#include <Templates/Function.h>
 #include <Templates/UnrealTemplate.h>
 #define UE_AUTORTFM_API CORE_API
 #define UE_AUTORTFM_FORCEINLINE FORCEINLINE
@@ -82,10 +84,6 @@
 #define UE_AUTORTFM_FORCEINLINE inline
 #endif
 
-#if UE_AUTORTFM
-template<typename FuncType> class TFunction;
-#endif
-
 #define UE_AUTORTFM_UNUSED(UNUSEDVAR) (void)UNUSEDVAR
 
 #ifdef __cplusplus
