[AutoRTFM] Fix data race in codegen for static local initializers which can cause crashes.

Move the calls to autortfm_pre_open() and autortfm_pre_static_local_initializer() to the correct basic block for the static initializer.

diff --git a/Engine/Source/Programs/AutoRTFMTests/Private/Statics.cpp b/Engine/Source/Programs/AutoRTFMTests/Private/Statics.cpp
--- a/Engine/Source/Programs/AutoRTFMTests/Private/Statics.cpp
+++ b/Engine/Source/Programs/AutoRTFMTests/Private/Statics.cpp
@@ -1,7 +1,15 @@
 // Copyright Epic Games, Inc. All Rights Reserved.
 
+#include "AutoRTFM/AutoRTFM.h"
+#include "AutoRTFM/Context.h"
+#include "AutoRTFM/ContextInlines.h"
+#include "AutoRTFM/Transaction.h"
 #include "Catch2Includes.h"
 #include <AutoRTFM/AutoRTFMTesting.h>
+#include <chrono>
+#include <condition_variable>
+#include <mutex>
+#include <thread>
 
 TEST_CASE("Statics.Int")
 {
@@ -292,3 +300,86 @@ TEST_CASE("Static.IsTransactional")
 		});
 	}
 }
+
+// Test case for SOL-7360
+TEST_CASE("Static.Concurrent")
+{
+	// A simple Signal() / Wait() event.
+	class FEvent
+	{
+	public:
+		UE_AUTORTFM_ALWAYS_OPEN void Signal()
+		{
+			{
+				std::unique_lock<std::mutex> Lock{Mutex};
+				Signalled = true;
+			}
+			CondVar.notify_all();
+		}
+
+		UE_AUTORTFM_ALWAYS_OPEN void Wait()
+		{
+			std::unique_lock<std::mutex> Lock{Mutex};
+			CondVar.wait(Lock, [&]{ return Signalled; });
+		}
+	private:
+		std::mutex Mutex;
+		std::condition_variable CondVar;
+		bool Signalled = false;
+	};
+
+	// Signalled just before the TransactionalThread waits on UnblockTransactionThread
+	FEvent TransactionThreadReady;
+	// Signalled when the main thread enters the static initializer for StaticInitFn::MyStatic
+	FEvent UnblockTransactionThread;
+
+	// A lambda function that holds a static local, which is called from the main thread
+	// and then the TransactionalThread.
+	auto StaticInitFn = [&]
+	{
+		// While StaticInitFn is called from both the main thread and the TransactionalThread,
+		// the initializer is only called on the main thread (the first invoker).
+		static bool MyStatic = [&]
+		{
+			// Wait until the TransactionalThread is ready.
+			TransactionThreadReady.Wait();
+			// Unblock the TransactionalThread.
+			UnblockTransactionThread.Signal();
+			// Pause the main thread for a short duration so that the 
+			// TransactionalThread can reach the static initializer guard.
+			std::this_thread::sleep_for(std::chrono::milliseconds(10));
+			return true;
+		}();
+	};
+
+	// A separate thread used to execute a static local initializer under a transaction.
+	std::thread TransactionalThread = std::thread([&]
+	{
+		AutoRTFM::Transact([&]
+		{
+			AutoRTFM::FTransaction* Transaction = AutoRTFM::FContext::Get()->GetCurrentTransaction();
+			REQUIRE(Transaction->State() == AutoRTFM::FTransaction::EState::ClosedActive);
+			// Signal to the main thread that this thread is about to wait on
+			// UnblockTransactionThread. This attempts to keep the threads
+			// tightly synchronized.
+			TransactionThreadReady.Signal();
+
+			// In order to trigger SOL-7360 we need the code in this block to
+			// execute quickly (within the main thread sleep duration):
+			{
+				// Wait until the main thread is in the static initializer.
+				UnblockTransactionThread.Wait();
+				// While the main thread is in the static initializer, attempt
+				// to lock the static initializer guard on this thread.
+				StaticInitFn();
+			}
+
+			// If SOL-7360 is not fixed, then the transaction state will be incorrect.
+			// (An assert would have likely been hit already too)
+			REQUIRE(Transaction->State() == AutoRTFM::FTransaction::EState::ClosedActive);
+		});
+	});
+
+	StaticInitFn();
+	TransactionalThread.join();
+}
