[AutoRTFM] Fix "Not allowed to construct a stack local within a transaction" check

It was allowing the thing it was supposed to forbid.
Also pop the on-abort handlers when destructing a TLockFreeFixedSizeAllocator. While it's *probably* harmless to use 'this' after destruction (because AutoRTFM will keep the memory alive until the transaction is closed), this is totally UB.

diff --git a/Engine/Source/Programs/AutoRTFMTests/Private/LockFreeFixedSizeAllocatorTests.cpp b/Engine/Source/Programs/AutoRTFMTests/Private/LockFreeFixedSizeAllocatorTests.cpp
--- a/Engine/Source/Programs/AutoRTFMTests/Private/LockFreeFixedSizeAllocatorTests.cpp
+++ b/Engine/Source/Programs/AutoRTFMTests/Private/LockFreeFixedSizeAllocatorTests.cpp
@@ -43,6 +43,21 @@ TEMPLATE_TEST_CASE("LockFreeFixedSizeAllocator", "", FThreadSafeCounter, FThread
 		REQUIRE(1 == Allocator.GetNumFree());
 	}
 
+	SECTION("Transact(new(ALockFreeFixedSizeAllocator), Allocate, delete(Allocator), Abort)")
+	{
+		void* Blob = nullptr;
+
+		AutoRTFM::Testing::Abort([&]
+		{
+			ALockFreeFixedSizeAllocator* Allocator = new ALockFreeFixedSizeAllocator();
+			Blob = Allocator->Allocate();
+			delete Allocator;
+			AutoRTFM::AbortTransaction();
+		});
+
+		REQUIRE(nullptr == Blob);
+	}
+
 	SECTION("Transact(Transact(Allocate, Abort), Allocate)")
 	{
 		ALockFreeFixedSizeAllocator Allocator;
diff --git a/Engine/Source/Runtime/Core/Public/Containers/LockFreeFixedSizeAllocator.h b/Engine/Source/Runtime/Core/Public/Containers/LockFreeFixedSizeAllocator.h
--- a/Engine/Source/Runtime/Core/Public/Containers/LockFreeFixedSizeAllocator.h
+++ b/Engine/Source/Runtime/Core/Public/Containers/LockFreeFixedSizeAllocator.h
@@ -197,12 +197,13 @@ class TLockFreeFixedSizeAllocator
 public:
 	TLockFreeFixedSizeAllocator()
 	{
-		checkf(AutoRTFM::IsClosed() || !AutoRTFM::IsOnCurrentTransactionStack(this), TEXT("Not allowed to construct a stack local within a transaction."));
+		checkf(!AutoRTFM::IsClosed() || !AutoRTFM::IsOnCurrentTransactionStack(this), TEXT("Not allowed to construct a stack local within a transaction."));
 	}
 
 	/** Destructor, returns all memory via FMemory::Free **/
 	~TLockFreeFixedSizeAllocator()
 	{
+		AutoRTFM::PopAllOnAbortHandlers(this);
 		UE_AUTORTFM_ONCOMMIT(this)
 		{
 			check(GetNumUsed() == 0);
diff --git a/Engine/Source/Runtime/CoreUObject/Public/AssetRegistry/AssetDataTagMapSerializationDetails.h b/Engine/Source/Runtime/CoreUObject/Public/AssetRegistry/AssetDataTagMapSerializationDetails.h
--- a/Engine/Source/Runtime/CoreUObject/Public/AssetRegistry/AssetDataTagMapSerializationDetails.h
+++ b/Engine/Source/Runtime/CoreUObject/Public/AssetRegistry/AssetDataTagMapSerializationDetails.h
@@ -86,7 +86,7 @@ namespace FixedTagPrivate
 
 		explicit FStore(uint32 InIndex) : Index(InIndex)
 		{
-			checkf(AutoRTFM::IsClosed() || !AutoRTFM::IsOnCurrentTransactionStack(this), TEXT("Not allowed to construct a stack local within a transaction."));
+			checkf(!AutoRTFM::IsClosed() || !AutoRTFM::IsOnCurrentTransactionStack(this), TEXT("Not allowed to construct a stack local within a transaction."));
 		}
 
 		~FStore();
