[AutoRTFM] Fix open writes to transactional memory in LockFreeLinkAllocator_TLSCache

TLockFreeFixedSizeAllocator is already AutoRTFM-safe, and all calls made by TLockFreeFixedSizeAllocator into LockFreeLinkAllocator_TLSCache are all in the open.
However - constructing or destructing a TLockFreePointerListFIFO in the closed will call into LockFreeLinkAllocator_TLSCache, which can record writes.
The mix of open and closed writes to LockFreeLinkAllocator_TLSCache can lead to memory corruption (SOL-6743).

Add AutoRTFM logic to LockFreeLinkAllocator_TLSCache so that this can be safely called from closed transactions.

diff --git a/Engine/Source/Programs/AutoRTFMTests/Private/LockFreeFixedSizeAllocatorTests.cpp b/Engine/Source/Programs/AutoRTFMTests/Private/LockFreeFixedSizeAllocatorTests.cpp
--- a/Engine/Source/Programs/AutoRTFMTests/Private/LockFreeFixedSizeAllocatorTests.cpp
+++ b/Engine/Source/Programs/AutoRTFMTests/Private/LockFreeFixedSizeAllocatorTests.cpp
@@ -1,8 +1,10 @@
 // Copyright Epic Games, Inc. All Rights Reserved.
 
+#include "AutoRTFM/AutoRTFM.h"
 #include "Catch2Includes.h"
 #include "AutoRTFM/AutoRTFMTesting.h"
 #include "Containers/LockFreeFixedSizeAllocator.h"
+#include "Containers/LockFreeList.h"
 #include "HAL/ThreadSafeCounter.h"
 #include "HAL/ThreadSafeCounter64.h"
 
@@ -265,4 +267,112 @@ TEMPLATE_TEST_CASE("LockFreeFixedSizeAllocator", "", FThreadSafeCounter, FThread
 
 		Allocator.Free(Blob);
 	}
+
+	SECTION("Transact(Open(Allocate, Free), Allocate, Free)")
+	{
+		ALockFreeFixedSizeAllocator Allocator;
+
+		AutoRTFM::Testing::Commit([&]
+			{
+				AutoRTFM::Open([&]
+				{
+					void* Blob = Allocator.Allocate();
+					Allocator.Free(Blob);
+				});
+
+				void* Blob = Allocator.Allocate();
+				Allocator.Free(Blob);
+			});
+
+		REQUIRE(0 == Allocator.GetNumUsed());
+		REQUIRE(1 == Allocator.GetNumFree());
+	}
+
+	SECTION("Transact(Allocate, Free, Open(Allocate, Free))")
+	{
+		ALockFreeFixedSizeAllocator Allocator;
+
+		AutoRTFM::Testing::Commit([&]
+			{
+				{
+					void* Blob = Allocator.Allocate();
+					Allocator.Free(Blob);
+				}
+				AutoRTFM::Open([&]
+				{
+					void* Blob = Allocator.Allocate();
+					Allocator.Free(Blob);
+				});
+			});
+
+		REQUIRE(0 == Allocator.GetNumUsed());
+		REQUIRE(2 == Allocator.GetNumFree());
+	}
+
+	SECTION("Allocate, Free, Transact(Open(Allocate, Free), Allocate, Free)")
+	{
+		ALockFreeFixedSizeAllocator Allocator;
+		{
+			void* Blob = Allocator.Allocate();
+			Allocator.Free(Blob);
+		}
+		AutoRTFM::Testing::Commit([&]
+			{
+				AutoRTFM::Open([&]
+				{
+					void* Blob = Allocator.Allocate();
+					Allocator.Free(Blob);
+				});
+
+				void* Blob = Allocator.Allocate();
+				Allocator.Free(Blob);
+			});
+
+		REQUIRE(0 == Allocator.GetNumUsed());
+		REQUIRE(1 == Allocator.GetNumFree());
+	}
+
+	SECTION("Allocate, Free, Transact(Allocate, Free, Open(Allocate, Free))")
+	{
+		ALockFreeFixedSizeAllocator Allocator;
+		{
+			void* Blob = Allocator.Allocate();
+			Allocator.Free(Blob);
+		}
+		AutoRTFM::Testing::Commit([&]
+			{
+				{
+					void* Blob = Allocator.Allocate();
+					Allocator.Free(Blob);
+				}
+				AutoRTFM::Open([&]
+				{
+					void* Blob = Allocator.Allocate();
+					Allocator.Free(Blob);
+				});
+			});
+
+		REQUIRE(0 == Allocator.GetNumUsed());
+		REQUIRE(2 == Allocator.GetNumFree());
+	}
+
+	// Specific test for SOL-7378
+	SECTION("Transact(TLockFreePointerListFIFO::Ctor, Open(Allocate, Free), TLockFreePointerListFIFO::Dtor)")
+	{
+		ALockFreeFixedSizeAllocator Allocator;
+
+		AutoRTFM::Testing::Commit([&]
+			{
+				TLockFreePointerListFIFO<int, 64> FIFO;
+				
+				AutoRTFM::Open([&]
+				{
+					void* Blob = Allocator.Allocate();
+					Allocator.Free(Blob);
+				});
+			});
+
+		REQUIRE(0 == Allocator.GetNumUsed());
+		REQUIRE(1 == Allocator.GetNumFree());
+	}
 }
diff --git a/Engine/Source/Runtime/Core/Private/Containers/LockFreeList.cpp b/Engine/Source/Runtime/Core/Private/Containers/LockFreeList.cpp
--- a/Engine/Source/Runtime/Core/Private/Containers/LockFreeList.cpp
+++ b/Engine/Source/Runtime/Core/Private/Containers/LockFreeList.cpp
@@ -1,6 +1,7 @@
 // Copyright Epic Games, Inc. All Rights Reserved.
 
 #include "Containers/LockFreeList.h"
+#include "AutoRTFM/AutoRTFM.h"
 #include "HAL/PlatformProcess.h"
 #include "HAL/IConsoleManager.h"
 #include "Stats/Stats.h"
@@ -85,15 +86,30 @@ public:
 
 	LockFreeLinkAllocator_TLSCache()
 	{
-		check(IsInGameThread());
-		TlsSlot = FPlatformTLS::AllocTlsSlot();
-		check(FPlatformTLS::IsValidTlsSlot(TlsSlot));
+		checkf(!AutoRTFM::IsClosed() || !AutoRTFM::IsOnCurrentTransactionStack(this), TEXT("Not allowed to construct a stack local within a transaction."));
+
+		UE_AUTORTFM_OPEN
+		{
+			check(IsInGameThread());
+			TlsSlot = FPlatformTLS::AllocTlsSlot();
+			check(FPlatformTLS::IsValidTlsSlot(TlsSlot));
+		};
+		AutoRTFM::PushOnAbortHandler(this, [this]
+		{
+			FPlatformTLS::FreeTlsSlot(TlsSlot);
+			TlsSlot = FPlatformTLS::InvalidTlsSlot;
+		});
 	}
+	
 	/** Destructor, leaks all of the memory **/
 	~LockFreeLinkAllocator_TLSCache()
 	{
-		FPlatformTLS::FreeTlsSlot(TlsSlot);
-		TlsSlot = FPlatformTLS::InvalidTlsSlot;
+		AutoRTFM::PopAllOnAbortHandlers(this);
+		UE_AUTORTFM_ONCOMMIT(this)
+		{
+			FPlatformTLS::FreeTlsSlot(TlsSlot);
+			TlsSlot = FPlatformTLS::InvalidTlsSlot;
+		};
 	}
 
 	/**
@@ -104,6 +120,19 @@ public:
 	*/
 	TLinkPtr Pop()
 	{
+		if (AutoRTFM::IsClosed())
+		{
+			TLinkPtr Ptr = AutoRTFM::Open([&]
+			{
+				return Pop();
+			});
+			AutoRTFM::PushOnAbortHandler(this, [this, Ptr]
+			{
+				Push(Ptr);
+			});
+			return Ptr;
+		}
+
 		FThreadLocalCache& TLS = GetTLS();
 
 		if (!TLS.PartialBundle)
@@ -149,6 +178,15 @@ public:
 	*/
 	void Push(TLinkPtr Item)
 	{
+		if (AutoRTFM::IsClosed())
+		{
+			UE_AUTORTFM_ONCOMMIT(this, Item)
+			{
+				return Push(Item);
+			};
+			return;
+		}
+
 		FThreadLocalCache& TLS = GetTLS();
 		if (TLS.NumPartial >= NUM_PER_BUNDLE)
 		{
