[AutoRTFM] Clean up `TRefCountingMixin` refcounting logic for multithreaded safety.

Previously, `TRefCountingMixin` used an always-`OnCommit` strategy for updating the refcount. This was not actually safe in the presence of multiple threads and could have led to premature deallocation.

We now use the same strategy as FRefCountBase, using `OnAbort` to undo `AddRef` and `OnCommit` to defer `Release`. The existing tests have been templatized.

diff --git a/Engine/Source/Programs/AutoRTFMTests/Private/RefCountPtrTests.cpp b/Engine/Source/Programs/AutoRTFMTests/Private/RefCountPtrTests.cpp
--- a/Engine/Source/Programs/AutoRTFMTests/Private/RefCountPtrTests.cpp
+++ b/Engine/Source/Programs/AutoRTFMTests/Private/RefCountPtrTests.cpp
@@ -11,18 +11,24 @@ namespace
 		RefCountedInt(int InValue) : Value(InValue) {}
 		int Value;
 	};
+
+	struct RefCountedMixinInt : public TRefCountingMixin<RefCountedMixinInt>
+	{
+		RefCountedMixinInt(int InValue) : Value(InValue) {}
+		int Value;
+	};
 }
 
-TEST_CASE("RefCountPtr.PreviouslyAllocated")
+TEMPLATE_TEST_CASE("PreviouslyAllocated", "[RefCountPtr]", RefCountedInt, RefCountedMixinInt)
 {
 	SECTION("Using new T()")
 	{
-		TRefCountPtr<RefCountedInt> Foo(new RefCountedInt(42));
+		TRefCountPtr<TestType> Foo(new TestType(42));
 
 		AutoRTFM::Commit([&]
 		{
 			// Make a copy to bump the reference count.
-			TRefCountPtr<RefCountedInt> Copy = Foo;
+			TRefCountPtr<TestType> Copy = Foo;
 			Copy->Value = 13;
 		});
 
@@ -32,12 +38,12 @@ TEST_CASE("RefCountPtr.PreviouslyAllocated")
 
 	SECTION("Using MakeRefCount<T>")
 	{
-		TRefCountPtr<RefCountedInt> Foo = MakeRefCount<RefCountedInt>(42);
+		TRefCountPtr<TestType> Foo = MakeRefCount<TestType>(42);
 
 		AutoRTFM::Commit([&]
 		{
 			// Make a copy to bump the reference count.
-			TRefCountPtr<RefCountedInt> Copy = Foo;
+			TRefCountPtr<TestType> Copy = Foo;
 			Copy->Value = 13;
 		});
 
@@ -46,16 +52,16 @@ TEST_CASE("RefCountPtr.PreviouslyAllocated")
 	}
 }
 
-TEST_CASE("RefCountPtr.AbortWithPreviouslyAllocated")
+TEMPLATE_TEST_CASE("AbortWithPreviouslyAllocated", "[RefCountPtr]", RefCountedInt, RefCountedMixinInt)
 {
 	SECTION("Using new T()")
 	{
-		TRefCountPtr<RefCountedInt> Foo(new RefCountedInt(42));
+		TRefCountPtr<TestType> Foo(new TestType(42));
 
 		REQUIRE(AutoRTFM::ETransactionResult::AbortedByRequest == AutoRTFM::Transact([&]
 		{
 			// Make a copy to bump the reference count.
-			TRefCountPtr<RefCountedInt> Copy = Foo;
+			TRefCountPtr<TestType> Copy = Foo;
 			Copy->Value = 13;
 
 			AutoRTFM::AbortTransaction();
@@ -67,12 +73,12 @@ TEST_CASE("RefCountPtr.AbortWithPreviouslyAllocated")
 
 	SECTION("Using MakeRefCount<T>")
 	{
-		TRefCountPtr<RefCountedInt> Foo = MakeRefCount<RefCountedInt>(42);
+		TRefCountPtr<TestType> Foo = MakeRefCount<TestType>(42);
 
 		REQUIRE(AutoRTFM::ETransactionResult::AbortedByRequest == AutoRTFM::Transact([&]
 		{
 			// Make a copy to bump the reference count.
-			TRefCountPtr<RefCountedInt> Copy = Foo;
+			TRefCountPtr<TestType> Copy = Foo;
 			Copy->Value = 13;
 
 			AutoRTFM::AbortTransaction();
@@ -83,7 +89,7 @@ TEST_CASE("RefCountPtr.AbortWithPreviouslyAllocated")
 	}
 }
 
-TEST_CASE("RefCountPtr.AbortWithNewlyAllocated")
+TEMPLATE_TEST_CASE("AbortWithNewlyAllocated", "[RefCountPtr]", RefCountedInt, RefCountedMixinInt)
 {
 	int Result = 42;
 
@@ -91,7 +97,7 @@ TEST_CASE("RefCountPtr.AbortWithNewlyAllocated")
 	{
 		REQUIRE(AutoRTFM::ETransactionResult::AbortedByRequest == AutoRTFM::Transact([&]
 		{
-			TRefCountPtr<RefCountedInt> Foo(new RefCountedInt(42));
+			TRefCountPtr<TestType> Foo(new TestType(42));
 			Result = Foo->Value;
 			AutoRTFM::AbortTransaction();
 		}));
@@ -101,8 +107,8 @@ TEST_CASE("RefCountPtr.AbortWithNewlyAllocated")
 	{
 		REQUIRE(AutoRTFM::ETransactionResult::AbortedByRequest == AutoRTFM::Transact([&]
 		{
-			TRefCountPtr<RefCountedInt> Foo(new RefCountedInt(42));
-			TRefCountPtr<RefCountedInt> Copy = Foo;
+			TRefCountPtr<TestType> Foo(new TestType(42));
+			TRefCountPtr<TestType> Copy = Foo;
 			Result = Copy->Value;
 			AutoRTFM::AbortTransaction();
 		}));
@@ -112,7 +118,7 @@ TEST_CASE("RefCountPtr.AbortWithNewlyAllocated")
 	{
 		REQUIRE(AutoRTFM::ETransactionResult::AbortedByRequest == AutoRTFM::Transact([&]
 		{
-			TRefCountPtr<RefCountedInt> Foo = MakeRefCount<RefCountedInt>(42);
+			TRefCountPtr<TestType> Foo = MakeRefCount<TestType>(42);
 			Result = Foo->Value;
 			AutoRTFM::AbortTransaction();
 		}));
@@ -122,8 +128,8 @@ TEST_CASE("RefCountPtr.AbortWithNewlyAllocated")
 	{
 		REQUIRE(AutoRTFM::ETransactionResult::AbortedByRequest == AutoRTFM::Transact([&]
 		{
-			TRefCountPtr<RefCountedInt> Foo = MakeRefCount<RefCountedInt>(42);
-			TRefCountPtr<RefCountedInt> Copy = Foo;
+			TRefCountPtr<TestType> Foo = MakeRefCount<TestType>(42);
+			TRefCountPtr<TestType> Copy = Foo;
 			Result = Copy->Value;
 			AutoRTFM::AbortTransaction();
 		}));
@@ -132,9 +138,9 @@ TEST_CASE("RefCountPtr.AbortWithNewlyAllocated")
 	REQUIRE(Result == 42);
 }
 
-TEST_CASE("RefCountPtr.OnCommitCapturingRefCountPtr")
+TEMPLATE_TEST_CASE("OnCommitCapturingRefCountPtr", "[RefCountPtr]", RefCountedInt, RefCountedMixinInt)
 {
-	TRefCountPtr<RefCountedInt> Foo = MakeRefCount<RefCountedInt>(42);
+	TRefCountPtr<TestType> Foo = MakeRefCount<TestType>(42);
 
 	SECTION("Committing")
 	{
@@ -164,9 +170,9 @@ TEST_CASE("RefCountPtr.OnCommitCapturingRefCountPtr")
 	}
 }
 
-TEST_CASE("RefCountPtr.OnAbortCapturingRefCountPtr")
+TEMPLATE_TEST_CASE("OnAbortCapturingRefCountPtr", "[RefCountPtr]", RefCountedInt, RefCountedMixinInt)
 {
-	TRefCountPtr<RefCountedInt> Foo = MakeRefCount<RefCountedInt>(42);
+	TRefCountPtr<TestType> Foo = MakeRefCount<TestType>(42);
 
 	SECTION("Committing")
 	{
diff --git a/Engine/Source/Runtime/Core/Public/Templates/RefCounting.h b/Engine/Source/Runtime/Core/Public/Templates/RefCounting.h
--- a/Engine/Source/Runtime/Core/Public/Templates/RefCounting.h
+++ b/Engine/Source/Runtime/Core/Public/Templates/RefCounting.h
@@ -266,19 +266,36 @@ public:
 	{
 		if constexpr (Mode == ERefCountingMode::ThreadSafe)
 		{
-			// Incrementing a reference count with relaxed ordering is always safe because no other action is taken
-			// in response to the increment, so there's nothing to order with.
+			uint32 Refs;
+			UE_AUTORTFM_OPEN
+			{
+				Refs = RefCount++;
+			};
 
-			AutoRTFM::OnCommit([this]
+			// If we are inside a transaction, a rollback must undo our refcount bump.
+			// In general, this is best handled by Release(). However, there is one case 
+			// that needs to be handled with special care. A brand-new object has a
+			// refcount of zero, and a rollback must return it to this zero-refcount state.
+			// However, calling AddRef() followed by Release() would not accomplish this;
+			// instead, it would free the object entirely! We need to guard against this,
+			// since it could lead to a double-free, so we detect the zero-reference case
+			// and special-case it.
+			if (Refs == 0)
 			{
-#if defined(_MSC_VER) && (defined(_M_X64) || defined(_M_IX86))
-				// We do a regular SC increment here because it maps to an _InterlockedIncrement (lock inc).
-				// The codegen for a relaxed fetch_add is actually much worse under MSVC (lock xadd).
-				++RefCount;
-#else
-				RefCount.fetch_add(1, std::memory_order_relaxed);
-#endif
-			});
+				UE_AUTORTFM_ONABORT(this)
+				{
+					--RefCount;
+					// The refcount is likely zero now, but leaving the object alive isn't a leak.
+					// We are restoring the object back to its initial "just-created" state.
+				};
+			}
+			else
+			{
+				UE_AUTORTFM_ONABORT(this)
+				{
+					Release();
+				};
+			}
 		}
 		else
 		{
@@ -291,26 +308,25 @@ public:
 
 	uint32 Release() const
 	{
+		checkSlow(RefCount > 0);
+
 		if constexpr (Mode == ERefCountingMode::ThreadSafe)
 		{
-			AutoRTFM::OnCommit([this]
+			UE_AUTORTFM_ONCOMMIT(this)
 			{
-				// std::memory_order_acq_rel is used here so that, if we do end up executing the destructor, it's not possible
-				// for side effects from executing the destructor end up being visible before we've determined that the 
-				// reference count is actually zero.
-
+				// std::memory_order_acq_rel is used here so that, if we do end up executing the
+				// destructor, it's not possible for side effects from executing the destructor end up
+				// being visible before we've determined that the reference count is actually zero.
 				uint32 OldCount = RefCount.fetch_sub(1, std::memory_order_acq_rel);
 				checkSlow(OldCount > 0);
 				if (OldCount == 1)
 				{
 					T::StaticDestroyObject(static_cast<const T*>(this));
 				}
-			});
+			};
 		}
 		else
 		{
-			checkSlow(RefCount > 0);
-
 			if (--RefCount == 0)
 			{
 				T::StaticDestroyObject(static_cast<const T*>(this));
