Implement a TransactionallySafe TryWriteLock
  - We will only try to grab the lock when we have a 0 Transactional count, inside a transaction
  - We will only setup the OnAbort if we did grab the lock, other we just return false

diff --git a/Engine/Source/Programs/AutoRTFMTests/Private/UECoreTests.cpp b/Engine/Source/Programs/AutoRTFMTests/Private/UECoreTests.cpp
--- a/Engine/Source/Programs/AutoRTFMTests/Private/UECoreTests.cpp
+++ b/Engine/Source/Programs/AutoRTFMTests/Private/UECoreTests.cpp
@@ -1437,6 +1437,126 @@ TEST_CASE("UECore.TransactionallySafeReadScopeLock")
 	}
 }
 
+TEST_CASE("UECore.TransactionallySafeRWLockTryWriteLock")
+{
+	SECTION("Outside Transaction")
+	{
+		FTransactionallySafeRWLock RWLock;
+
+		AutoRTFM::ETransactionResult Result = AutoRTFM::Transact([&]
+			{
+				REQUIRE(RWLock.TryWriteLock());
+				AutoRTFM::AbortTransaction();
+			});
+
+		REQUIRE(AutoRTFM::ETransactionResult::AbortedByRequest == Result);
+
+		Result = AutoRTFM::Transact([&]
+			{
+				REQUIRE(RWLock.TryWriteLock());
+				RWLock.WriteUnlock();
+			});
+
+		REQUIRE(AutoRTFM::ETransactionResult::Committed == Result);
+	}
+
+	SECTION("Inside Transaction")
+	{
+		AutoRTFM::ETransactionResult Result = AutoRTFM::Transact([&]
+			{
+				FTransactionallySafeRWLock RWLock;
+				REQUIRE(RWLock.TryWriteLock());
+				AutoRTFM::AbortTransaction();
+			});
+
+		REQUIRE(AutoRTFM::ETransactionResult::AbortedByRequest == Result);
+
+		Result = AutoRTFM::Transact([&]
+			{
+				FTransactionallySafeRWLock RWLock;
+				REQUIRE(RWLock.TryWriteLock());
+				RWLock.WriteUnlock();
+			});
+
+		REQUIRE(AutoRTFM::ETransactionResult::Committed == Result);
+	}
+
+	SECTION("Inside Transaction Used In Nested Transaction")
+	{
+		AutoRTFM::ETransactionResult InnerResult = AutoRTFM::ETransactionResult::Committed;
+
+		AutoRTFM::ETransactionResult Result = AutoRTFM::Transact([&]
+			{
+				FTransactionallySafeRWLock RWLock;
+
+				InnerResult = AutoRTFM::Transact([&]
+					{
+						REQUIRE(RWLock.TryWriteLock());
+						AutoRTFM::CascadingAbortTransaction();
+					});
+			});
+
+		REQUIRE(AutoRTFM::ETransactionResult::AbortedByCascade == Result);
+
+		Result = AutoRTFM::Transact([&]
+			{
+				FTransactionallySafeRWLock RWLock;
+
+				InnerResult = AutoRTFM::Transact([&]
+					{
+						REQUIRE(RWLock.TryWriteLock());
+						AutoRTFM::AbortTransaction();
+					});
+			});
+
+		REQUIRE(AutoRTFM::ETransactionResult::AbortedByRequest == InnerResult);
+		REQUIRE(AutoRTFM::ETransactionResult::Committed == Result);
+
+		Result = AutoRTFM::Transact([&]
+			{
+				FTransactionallySafeRWLock RWLock;
+
+				InnerResult = AutoRTFM::Transact([&]
+					{
+						REQUIRE(RWLock.TryWriteLock());
+					});
+
+				AutoRTFM::AbortTransaction();
+			});
+
+		REQUIRE(AutoRTFM::ETransactionResult::AbortedByRequest == Result);
+
+		Result = AutoRTFM::Transact([&]
+			{
+				FTransactionallySafeRWLock RWLock;
+
+				InnerResult = AutoRTFM::Transact([&]
+					{
+						REQUIRE(RWLock.TryWriteLock());
+					});
+
+				RWLock.WriteUnlock();
+			});
+
+		REQUIRE(AutoRTFM::ETransactionResult::Committed == InnerResult);
+		REQUIRE(AutoRTFM::ETransactionResult::Committed == Result);
+	}
+
+	SECTION("Multiple TryWriteLocks In Transaction")
+	{
+		AutoRTFM::ETransactionResult Result = AutoRTFM::Transact([&]
+			{
+				FTransactionallySafeRWLock RWLock;
+				REQUIRE(RWLock.TryWriteLock());
+				REQUIRE(RWLock.TryWriteLock());
+				RWLock.WriteUnlock();
+			});
+
+		REQUIRE(AutoRTFM::ETransactionResult::Committed == Result);
+	}
+}
+
+
 TEST_CASE("UECore.TransactionallySafeWriteScopeLock")
 {
 	SECTION("Outside Transaction")
diff --git a/Engine/Source/Runtime/Core/Public/Misc/TransactionallySafeRWLock.h b/Engine/Source/Runtime/Core/Public/Misc/TransactionallySafeRWLock.h
--- a/Engine/Source/Runtime/Core/Public/Misc/TransactionallySafeRWLock.h
+++ b/Engine/Source/Runtime/Core/Public/Misc/TransactionallySafeRWLock.h
@@ -129,14 +129,7 @@ struct FTransactionallySafeRWLockDefinition
 			// allocated and has already died before the on-abort is hit.
 			AutoRTFM::OnAbort([State = this->State]
 				{
-					ensure(0 != State->TransactionalLockCount);
-
-					State->TransactionalLockCount -= 1;
-
-					if (0 == State->TransactionalLockCount)
-					{
-						State->Lock.WriteUnlock();
-					}
+					State->WriteUnlock();
 				});
 		}
 		else
@@ -154,14 +147,7 @@ struct FTransactionallySafeRWLockDefinition
 			// allocated and has already died before the on-commit is hit.
 			AutoRTFM::OnCommit([State = this->State]
 				{
-					ensure(0 != State->TransactionalLockCount);
-
-					State->TransactionalLockCount -= 1;
-
-					if (0 == State->TransactionalLockCount)
-					{
-						State->Lock.WriteUnlock();
-					}
+					State->WriteUnlock();
 				});
 		}
 		else
@@ -173,7 +159,42 @@ struct FTransactionallySafeRWLockDefinition
 
 	bool TryWriteLock()
 	{
-		return false;
+		if (AutoRTFM::IsTransactional() || AutoRTFM::IsCommittingOrAborting())
+		{
+			bool Ret = false;
+
+			// Memory validation disabled as this is hot code.
+			UE_AUTORTFM_OPEN_NO_VALIDATION
+				{
+					// The transactional system which can increment TransactionalLockCount
+					// is always single-threaded, thus this is safe to check without atomicity.
+					// For TryWriteLock we should only lock when we have a 0 count as no one owns this lock
+					if (0 == State->TransactionalLockCount)
+					{
+						Ret = State->Lock.TryWriteLock();
+					}
+
+					if (Ret)
+					{
+						State->TransactionalLockCount += 1;
+					}
+				};
+
+			// Only setup the OnAbort if we *did* grab a lock, otherwise we will not want to do anything with the count or lock
+			if (Ret)
+			{
+				// We explicitly copy the state here for the case that `this` was stack
+				// allocated and has already died before the on-abort is hit.
+				AutoRTFM::OnAbort([State = this->State]
+					{
+						State->WriteUnlock();
+					});
+			}
+
+			return State->TransactionalLockCount != 0;
+		}
+
+		return State->Lock.TryWriteLock();
 	}
 
 private:
@@ -183,6 +204,18 @@ private:
 	{
 		FRWLock Lock;
 		uint32 TransactionalLockCount = 0;
+
+		void WriteUnlock()
+		{
+			ensure(0 != TransactionalLockCount);
+
+			TransactionalLockCount -= 1;
+
+			if (0 == TransactionalLockCount)
+			{
+				Lock.WriteUnlock();
+			}
+		}
 	};
 
 	TSharedPtr<FState> State;
