[AutoRTFM] Make FAssetDataTagMapSharedView transactionally safe

diff --git a/Engine/Source/Programs/AutoRTFMTests/Private/UECoreTests.cpp b/Engine/Source/Programs/AutoRTFMTests/Private/UECoreTests.cpp
--- a/Engine/Source/Programs/AutoRTFMTests/Private/UECoreTests.cpp
+++ b/Engine/Source/Programs/AutoRTFMTests/Private/UECoreTests.cpp
@@ -1,6 +1,7 @@
 // Copyright Epic Games, Inc. All Rights Reserved.
 
 #include "Catch2Includes.h"
+#include "AssetRegistry/AssetDataTagMap.h"
 #include "AutoRTFM/AutoRTFMTesting.h"
 #include "Delegates/IDelegateInstance.h"
 #include "HAL/MallocLeakDetection.h"
@@ -2201,3 +2202,143 @@ TEST_CASE("UECore.PropertyBag")
 		SS->DestroyStruct(Data);
 	});
 }
+
+TEST_CASE("UECore.FAssetDataTagMapSharedView")
+{
+	SECTION("Loose")
+	{
+		FAssetDataTagMap Loose;
+		Loose.Add(FName("cat"), FString("meow"));
+		Loose.Add(FName("dog"), FString("woof"));
+		SECTION("Copy FAssetDataTagMapSharedView from open")
+		{
+			FAssetDataTagMapSharedView Original{MoveTemp(Loose)};
+			SECTION("Commit")
+			{
+				AutoRTFM::Testing::Commit([&]
+				{
+					FAssetDataTagMapSharedView View{Original};
+				});
+			}
+			SECTION("Abort")
+			{
+				AutoRTFM::Testing::Abort([&]
+				{
+					FAssetDataTagMapSharedView View{Original};
+					AutoRTFM::AbortTransaction();
+				});
+			}
+			REQUIRE(Original.Contains(FName("cat")));
+			REQUIRE(Original.Contains(FName("dog")));
+		}
+		SECTION("Copy FAssetDataTagMapSharedView from closed")
+		{
+			SECTION("Commit")
+			{
+				AutoRTFM::Testing::Commit([&]
+				{
+					FAssetDataTagMapSharedView Original{MoveTemp(Loose)};
+					FAssetDataTagMapSharedView View{Original};
+					REQUIRE(View.Contains(FName("cat")));
+					REQUIRE(View.Contains(FName("dog")));
+				});
+			}
+			SECTION("Abort")
+			{
+				AutoRTFM::Testing::Abort([&]
+				{
+					FAssetDataTagMapSharedView Original{MoveTemp(Loose)};
+					FAssetDataTagMapSharedView View{Original};
+					AutoRTFM::AbortTransaction();
+				});
+			}
+		}
+		SECTION("Move FAssetDataTagMapSharedView from open")
+		{
+			FAssetDataTagMapSharedView Original{MoveTemp(Loose)};
+			SECTION("Commit")
+			{
+				AutoRTFM::Testing::Commit([&]
+				{
+					FAssetDataTagMapSharedView View{MoveTemp(Original)};
+					REQUIRE(View.Contains(FName("cat")));
+					REQUIRE(View.Contains(FName("dog")));
+				});
+			}
+			SECTION("Abort")
+			{
+				AutoRTFM::Testing::Abort([&]
+				{
+					FAssetDataTagMapSharedView View{MoveTemp(Original)};
+					AutoRTFM::AbortTransaction();
+				});
+				REQUIRE(Original.Contains(FName("cat")));
+				REQUIRE(Original.Contains(FName("dog")));
+			}
+		}
+		SECTION("Move FAssetDataTagMapSharedView from closed")
+		{
+			SECTION("Commit")
+			{
+				AutoRTFM::Testing::Commit([&]
+				{
+					FAssetDataTagMapSharedView Original{MoveTemp(Loose)};
+					FAssetDataTagMapSharedView View{MoveTemp(Original)};
+					REQUIRE(View.Contains(FName("cat")));
+					REQUIRE(View.Contains(FName("dog")));
+				});
+			}
+			SECTION("Abort")
+			{
+				AutoRTFM::Testing::Abort([&]
+				{
+					FAssetDataTagMapSharedView Original{MoveTemp(Loose)};
+					FAssetDataTagMapSharedView View{MoveTemp(Original)};
+					AutoRTFM::AbortTransaction();
+				});
+			}
+		}
+		SECTION("Move FAssetDataTagMap from open")
+		{
+			SECTION("Commit")
+			{
+				AutoRTFM::Testing::Commit([&]
+				{
+					FAssetDataTagMapSharedView View{MoveTemp(Loose)};
+					REQUIRE(View.Contains(FName("cat")));
+					REQUIRE(View.Contains(FName("dog")));
+				});
+			}
+			SECTION("Abort")
+			{
+				AutoRTFM::Testing::Abort([&]
+				{
+					FAssetDataTagMapSharedView View{MoveTemp(Loose)};
+					AutoRTFM::AbortTransaction();
+				});
+			}
+		}
+		SECTION("Move FAssetDataTagMap from closed")
+		{
+			SECTION("Commit")
+			{
+				AutoRTFM::Testing::Commit([&]
+				{
+					FAssetDataTagMap ClosedLoose{MoveTemp(Loose)};
+					FAssetDataTagMapSharedView View{MoveTemp(ClosedLoose)};
+					REQUIRE(View.Contains(FName("cat")));
+					REQUIRE(View.Contains(FName("dog")));
+				});
+			}
+			SECTION("Abort")
+			{
+				AutoRTFM::Testing::Abort([&]
+				{
+					FAssetDataTagMap ClosedLoose{MoveTemp(Loose)};
+					FAssetDataTagMapSharedView View{MoveTemp(ClosedLoose)};
+					AutoRTFM::AbortTransaction();
+				});
+			}
+		}
+	}
+}
diff --git a/Engine/Source/Runtime/CoreUObject/Private/AssetRegistry/AssetDataTagMap.cpp b/Engine/Source/Runtime/CoreUObject/Private/AssetRegistry/AssetDataTagMap.cpp
--- a/Engine/Source/Runtime/CoreUObject/Private/AssetRegistry/AssetDataTagMap.cpp
+++ b/Engine/Source/Runtime/CoreUObject/Private/AssetRegistry/AssetDataTagMap.cpp
@@ -4,6 +4,7 @@
 
 #include "Algo/Sort.h"
 #include "AssetRegistry/AssetDataTagMapSerializationDetails.h"
+#include "AutoRTFM/AutoRTFM.h"
 #include "Misc/PackageName.h"
 #include "Misc/Optional.h"
 #include "Misc/TransactionallySafeCriticalSection.h"
@@ -1446,11 +1447,29 @@ FAssetDataTagMapSharedView::FAssetDataTagMapSharedView(const FAssetDataTagMapSha
 {
 	if (IsFixed())
 	{
-		FixedTagPrivate::GStores[Fixed.StoreIndex].AddRef();
+		FixedTagPrivate::FStore* Store = &FixedTagPrivate::GStores[Fixed.StoreIndex];
+		UE_AUTORTFM_OPEN
+		{
+			Store->AddRef();
+		};
+		UE_AUTORTFM_ONABORT(Store)
+		{
+			Store->Release();
+		};
 	}
 	else if (IsLoose())
 	{
-		Loose->RefCount.Increment(); //-V614
+		UE_AUTORTFM_OPEN
+		{
+			Loose->RefCount.Increment(); //-V614
+		};
+		UE_AUTORTFM_ONABORT(Loose = this->Loose)
+		{
+			if (Loose->RefCount.Decrement() == 0)
+			{
+				delete Loose;
+			}
+		};
 	}
 }
 
@@ -1463,7 +1482,15 @@ FAssetDataTagMapSharedView::FAssetDataTagMapSharedView(FAssetDataTagMapSharedVie
 FAssetDataTagMapSharedView::FAssetDataTagMapSharedView(FixedTagPrivate::FMapHandle InFixed)
 	: Fixed(InFixed)
 {
-	FixedTagPrivate::GStores[Fixed.StoreIndex].AddRef();
+	FixedTagPrivate::FStore* Store = &FixedTagPrivate::GStores[Fixed.StoreIndex];
+	UE_AUTORTFM_OPEN
+	{
+		Store->AddRef();
+	};
+	UE_AUTORTFM_ONABORT(Store)
+	{
+		Store->Release();
+	};
 }
 
 FAssetDataTagMapSharedView::FAssetDataTagMapSharedView(FAssetDataTagMap&& InLoose)
@@ -1471,7 +1498,10 @@ FAssetDataTagMapSharedView::FAssetDataTagMapSharedView(FAssetDataTagMap&& InLoos
 	if (InLoose.Num())
 	{
 		Loose = new FAssetDataTagMap(MoveTemp(InLoose));
-		Loose->RefCount.Increment();
+		UE_AUTORTFM_OPEN
+		{
+			Loose->RefCount.Increment();
+		};
 	}
 }
 
@@ -1493,14 +1523,21 @@ FAssetDataTagMapSharedView::~FAssetDataTagMapSharedView()
 {
 	if (IsFixed())
 	{
-		FixedTagPrivate::GStores[Fixed.StoreIndex].Release();
+		FixedTagPrivate::FStore* Store = &FixedTagPrivate::GStores[Fixed.StoreIndex];
+		UE_AUTORTFM_ONCOMMIT(Store)
+		{
+			Store->Release();
+		};
 	}
 	else if (IsLoose())
 	{
-		if (Loose->RefCount.Decrement() == 0)
+		UE_AUTORTFM_ONCOMMIT(Loose = this->Loose)
 		{
-			delete Loose;
-		}
+			if (Loose->RefCount.Decrement() == 0)
+			{
+				delete Loose;
+			}
+		};
 	}
 }
 
