[AutoRTFM] Register closed function _tcsnlen to map to self

Add tests for FGuid::ImportTextItem(), which calls _tcsnlen on windows.

Also change the default of EnsureOnAbortByLanguage to true for the AutoRTFMTests, and scope-set it to false for the two tests that want to test aborting by language. This provides callstacks when closed functions are missing.

diff --git a/Engine/Source/Programs/AutoRTFMTests/Private/Abort.cpp b/Engine/Source/Programs/AutoRTFMTests/Private/Abort.cpp
--- a/Engine/Source/Programs/AutoRTFMTests/Private/Abort.cpp
+++ b/Engine/Source/Programs/AutoRTFMTests/Private/Abort.cpp
@@ -1,6 +1,7 @@
 // Copyright Epic Games, Inc. All Rights Reserved.
 
 #include "Catch2Includes.h"
+#include "Utils.h"
 #include <AutoRTFM/AutoRTFMTesting.h>
 #include <map>
 #include <vector>
@@ -635,6 +636,8 @@ static void FnHasNoClosed()
 
 TEST_CASE("Abort.Language")
 {
+	AUTORTFM_SCOPED_SET_ENSURE_ON_ABORT_BY_LANGUAGE(false);
+
 	bool bTouched = false;
 
 	const AutoRTFM::ETransactionResult Result = AutoRTFM::Transact([&]
@@ -649,6 +652,8 @@ TEST_CASE("Abort.Language")
 
 TEST_CASE("Abort.LanguageThroughOpen")
 {
+	AUTORTFM_SCOPED_SET_ENSURE_ON_ABORT_BY_LANGUAGE(false);
+
 	bool bTouched = false;
 
 	const AutoRTFM::ETransactionResult Result = AutoRTFM::Transact([&]
diff --git a/Engine/Source/Programs/AutoRTFMTests/Private/AutoRTFMTests.cpp b/Engine/Source/Programs/AutoRTFMTests/Private/AutoRTFMTests.cpp
--- a/Engine/Source/Programs/AutoRTFMTests/Private/AutoRTFMTests.cpp
+++ b/Engine/Source/Programs/AutoRTFMTests/Private/AutoRTFMTests.cpp
@@ -97,8 +97,9 @@ int RunTests(int ArgC, const char* ArgV[])
 	// Enable AutoRTFM.
 	AutoRTFM::ForTheRuntime::SetAutoRTFMRuntime(AutoRTFM::ForTheRuntime::EAutoRTFMEnabledState::AutoRTFM_Enabled);
 
-	// We don't want to trigger ensure's on abort because we are going to test that.
-	AutoRTFM::ForTheRuntime::SetEnsureOnAbortByLanguage(false);
+	// By default, abort by language. Tests can override this with 
+	// AUTORTFM_SCOPED_SET_ENSURE_ON_ABORT_BY_LANGUAGE(false)
+	AutoRTFM::ForTheRuntime::SetEnsureOnAbortByLanguage(true);
 
 	// Unexpected memory validation errors should be full assertions.
 	AutoRTFM::ForTheRuntime::SetMemoryValidationLevel(AutoRTFM::EMemoryValidationLevel::Error);
diff --git a/Engine/Source/Programs/AutoRTFMTests/Private/UEPlatform.cpp b/Engine/Source/Programs/AutoRTFMTests/Private/UEPlatform.cpp
--- a/Engine/Source/Programs/AutoRTFMTests/Private/UEPlatform.cpp
+++ b/Engine/Source/Programs/AutoRTFMTests/Private/UEPlatform.cpp
@@ -1,37 +1,34 @@
 // Copyright Epic Games, Inc. All Rights Reserved.
 
+#include "AutoRTFM/AutoRTFMTesting.h"
 #include "Catch2Includes.h"
 #include "AutoRTFM/AutoRTFM.h"
 #include "HAL/PlatformMisc.h"
+#include "Misc/Guid.h"
 #include "Utils.h"
 
 TEST_CASE("FPlatformMisc.CreateGUID")
 {
-	AUTORTFM_SCOPED_SET_ENSURE_ON_ABORT_BY_LANGUAGE(true);
+	FGuid Guid;
 
 	SECTION("Commit")
 	{
-		FGuid Guid;
 		SECTION("GuidOutsideTransactionStack")
 		{
-			AutoRTFM::ETransactionResult Result = AutoRTFM::Transact([&]
+			AutoRTFM::Testing::Commit([&]
 			{
 				FPlatformMisc::CreateGuid(Guid);
 			});
-
-			REQUIRE(AutoRTFM::ETransactionResult::Committed == Result);
 		}
 
 		SECTION("GuidInsideTransactionStack")
 		{
-			AutoRTFM::ETransactionResult Result = AutoRTFM::Transact([&]
+			AutoRTFM::Testing::Commit([&]
 			{
 				FGuid InnerGuid;
 				FPlatformMisc::CreateGuid(InnerGuid);
 				AutoRTFM::Open([&]{ Guid = InnerGuid; });
 			});
-
-			REQUIRE(AutoRTFM::ETransactionResult::Committed == Result);
 		}
 
 		REQUIRE(Guid != FGuid());
@@ -39,31 +36,61 @@ TEST_CASE("FPlatformMisc.CreateGUID")
 
 	SECTION("Abort")
 	{
-		FGuid Guid;
 		SECTION("GuidOutsideTransactionStack")
 		{
-			AutoRTFM::ETransactionResult Result = AutoRTFM::Transact([&]
+			AutoRTFM::Testing::Abort([&]
 			{
 				FPlatformMisc::CreateGuid(Guid);
 				AutoRTFM::AbortTransaction();
 			});
-
-			REQUIRE(AutoRTFM::ETransactionResult::AbortedByRequest == Result);
 		}
 
 		SECTION("GuidInsideTransactionStack")
 		{
-			AutoRTFM::ETransactionResult Result = AutoRTFM::Transact([&]
+			AutoRTFM::Testing::Abort([&]
 			{
 				FGuid InnerGuid;
 				FPlatformMisc::CreateGuid(InnerGuid);
 				Guid = InnerGuid;
 				AutoRTFM::AbortTransaction();
 			});
-
-			REQUIRE(AutoRTFM::ETransactionResult::AbortedByRequest == Result);
 		}
 
 		REQUIRE(Guid == FGuid());
 	}
 }
+
+TEST_CASE("FGuid.ImportTextItem")
+{
+	TCHAR Buffer[64] = TEXT("f06250a3d866649e3b3d77f936fe6620");
+	const TCHAR* GuidString = Buffer;
+	FGuid Expected;
+	REQUIRE(FGuid::Parse(GuidString, Expected));
+	
+	FGuid Guid;
+
+	SECTION("Commit")
+	{
+		AutoRTFM::Testing::Commit([&]
+		{
+			bool Result = Guid.ImportTextItem(GuidString, 0, nullptr, nullptr);
+			REQUIRE(Result == true);
+		});
+
+		REQUIRE((GuidString - Buffer) == 32);
+		REQUIRE(Guid == Expected);
+	}
+
+	SECTION("Abort")
+	{
+		AutoRTFM::Testing::Abort([&]
+		{
+			bool Result = Guid.ImportTextItem(GuidString, 0, nullptr, nullptr);
+			REQUIRE(Result == true);
+			AutoRTFM::AbortTransaction();
+		});
+
+		REQUIRE(GuidString == Buffer);
+		REQUIRE(Guid == FGuid());
+	}
+}
diff --git a/Engine/Source/Runtime/Core/Private/AutoRTFM/StdLibSupport.cpp b/Engine/Source/Runtime/Core/Private/AutoRTFM/StdLibSupport.cpp
--- a/Engine/Source/Runtime/Core/Private/AutoRTFM/StdLibSupport.cpp
+++ b/Engine/Source/Runtime/Core/Private/AutoRTFM/StdLibSupport.cpp
@@ -678,6 +678,7 @@ UE_AUTORTFM_REGISTER_SELF_OVERLOADED_FUNCTION(double (double, double), pow);
 #if PLATFORM_WINDOWS
 UE_AUTORTFM_REGISTER_SELF_FUNCTION(_tcsncmp);
 UE_AUTORTFM_REGISTER_SELF_FUNCTION(_tcslen);
+UE_AUTORTFM_REGISTER_SELF_FUNCTION(_tcsnlen);
 UE_AUTORTFM_REGISTER_SELF_FUNCTION(_isnan);
 UE_AUTORTFM_REGISTER_SELF_FUNCTION(_fdtest);
 UE_AUTORTFM_REGISTER_SELF_FUNCTION(_dtest);
