Add push/pop on-commit handlers, and fix a bug with how these are handled between transaction boundaries.

diff --git a/Engine/Source/Programs/AutoRTFMTests/Private/Abort.cpp b/Engine/Source/Programs/AutoRTFMTests/Private/Abort.cpp
--- a/Engine/Source/Programs/AutoRTFMTests/Private/Abort.cpp
+++ b/Engine/Source/Programs/AutoRTFMTests/Private/Abort.cpp
@@ -591,6 +591,98 @@ TEST_CASE("Abort.PushOnAbortHandler_Order")
 			REQUIRE(AutoRTFM::ETransactionResult::AbortedByRequest == Result);
 			REQUIRE(Value == 41);
 		}
+
+		SECTION("PopInChild")
+		{
+			int Value = 99;
+			AutoRTFM::ETransactionResult Result = AutoRTFM::Transact([&]
+			{
+				AutoRTFM::PushOnAbortHandler(&SAbortKey, [&Value] { Value = 42; });
+				AutoRTFM::Commit([&] { AutoRTFM::PopOnAbortHandler(&SAbortKey); });
+				AutoRTFM::AbortTransaction();
+			});
+
+			REQUIRE(AutoRTFM::ETransactionResult::AbortedByRequest == Result);
+			REQUIRE(99 == Value);
+		}
+
+		SECTION("PopInChildAndAbort")
+		{
+			int Value = 99;
+			AutoRTFM::ETransactionResult Result = AutoRTFM::Transact([&]
+			{
+				AutoRTFM::PushOnAbortHandler(&SAbortKey, [&Value] { Value = 42; });
+
+				AutoRTFM::Transact([&]
+				{
+					AutoRTFM::PopOnAbortHandler(&SAbortKey);
+
+					// This abort means the pop on abort handler *should not* propagate
+					// to the parent scope (meaning the parents push on abort should
+					// run as normal).
+					AutoRTFM::AbortTransaction();
+				});
+
+				AutoRTFM::AbortTransaction();
+			});
+
+			REQUIRE(AutoRTFM::ETransactionResult::AbortedByRequest == Result);
+			REQUIRE(42 == Value);
+		}
+
+		SECTION("PopInChildsChild")
+		{
+			int Value = 99;
+			AutoRTFM::ETransactionResult Result = AutoRTFM::Transact([&]
+			{
+				AutoRTFM::PushOnAbortHandler(&SAbortKey, [&Value] { Value = 42; });
+				AutoRTFM::Commit([&] { AutoRTFM::Commit([&] { AutoRTFM::PopOnAbortHandler(&SAbortKey); }); });
+				AutoRTFM::AbortTransaction();
+			});
+
+			REQUIRE(AutoRTFM::ETransactionResult::AbortedByRequest == Result);
+			REQUIRE(99 == Value);
+		}
+
+		SECTION("PopAllInChild")
+		{
+			int Value = 99;
+			AutoRTFM::ETransactionResult Result = AutoRTFM::Transact([&]
+			{
+				AutoRTFM::PushOnAbortHandler(&SAbortKey, [&Value] { Value = 42; });
+				AutoRTFM::Commit([&]
+				{
+					AutoRTFM::PushOnAbortHandler(&SAbortKey, [&Value] { Value = 13; });
+					AutoRTFM::PopAllOnAbortHandlers(&SAbortKey);
+				});
+				AutoRTFM::AbortTransaction();
+			});
+
+			REQUIRE(AutoRTFM::ETransactionResult::AbortedByRequest == Result);
+			REQUIRE(99 == Value);
+		}
+
+		SECTION("PopAllInChildAbort")
+		{
+			int Value = 99;
+			AutoRTFM::ETransactionResult Result = AutoRTFM::Transact([&]
+			{
+				AutoRTFM::PushOnAbortHandler(&SAbortKey, [&Value] { Value = 42; });
+				AutoRTFM::Transact([&]
+				{
+					AutoRTFM::PushOnAbortHandler(&SAbortKey, [&Value] { Value = 13; });
+					AutoRTFM::PopAllOnAbortHandlers(&SAbortKey);
+
+					// This abort will ensure that the pop-all cannot affect the outer transactions
+					// push on abort!
+					AutoRTFM::AbortTransaction();
+				});
+				AutoRTFM::AbortTransaction();
+			});
+
+			REQUIRE(AutoRTFM::ETransactionResult::AbortedByRequest == Result);
+			REQUIRE(42 == Value);
+		}
 	}
 }
 
diff --git a/Engine/Source/Programs/AutoRTFMTests/Private/Commit.cpp b/Engine/Source/Programs/AutoRTFMTests/Private/Commit.cpp
--- a/Engine/Source/Programs/AutoRTFMTests/Private/Commit.cpp
+++ b/Engine/Source/Programs/AutoRTFMTests/Private/Commit.cpp
@@ -1,7 +1,7 @@
 // Copyright Epic Games, Inc. All Rights Reserved.
 
 #include "Catch2Includes.h"
-#include <AutoRTFM/AutoRTFM.h>
+#include <AutoRTFM/AutoRTFMTesting.h>
 #include <map>
 #include <vector>
 
@@ -85,3 +85,343 @@ TEST_CASE("Commit.Large")
         AutoRTFM::ForTheRuntime::CheckConsistencyAssumingNoRaces();
     });
 }
+
+static int SCommitKey;
+
+TEST_CASE("Commit.PushOnCommitHandler_NoAbort")
+{
+	int Value = 55;
+
+	AutoRTFM::Testing::Commit([&]
+	{
+		Value = 66;
+		AutoRTFM::PushOnCommitHandler(&SCommitKey, [&Value]() { Value = 77; });
+		Value = 88;
+	});
+
+	REQUIRE(Value == 77);
+}
+
+TEST_CASE("Commit.PushOnCommitHandler_WithAbort")
+{
+	AutoRTFM::Testing::Abort([&]
+	{
+		AutoRTFM::PushOnCommitHandler(&SCommitKey, []() { FAIL("Unreachable"); });
+		AutoRTFM::AbortTransaction();
+	});
+}
+
+TEST_CASE("Commit.PushOnCommitHandler_MutableCapture")
+{
+	FString Message = "Hello";
+
+	AutoRTFM::Testing::Commit([&]
+	{
+		AutoRTFM::PushOnCommitHandler(&SCommitKey, [&Message]() mutable
+		{
+			Message += " World!";
+		});
+	});
+
+	REQUIRE(Message == "Hello World!");
+}
+
+TEST_CASE("Commit.PushOnCommitHandler_WithPop_NoAbort")
+{
+	AutoRTFM::Testing::Commit([&]
+	{
+		AutoRTFM::PushOnCommitHandler(&SCommitKey, []() { FAIL("Unreachable!"); });
+		AutoRTFM::PopOnCommitHandler(&SCommitKey);
+	});
+}
+
+TEST_CASE("Commit.PushOnCommitHandler_WithPopAll_NoAbort")
+{
+	AutoRTFM::Commit([&]
+	{
+		AutoRTFM::PushOnCommitHandler(&SCommitKey, []() { FAIL("Unreachable!"); });
+		AutoRTFM::PopAllOnCommitHandlers(&SCommitKey);
+	});
+}
+
+TEST_CASE("Commit.PushOnCommitHandler_WithPop_WithAbort")
+{
+	AutoRTFM::Testing::Abort([&]
+	{
+		AutoRTFM::PushOnCommitHandler(&SCommitKey, []() { FAIL("Unreachable!"); });
+		AutoRTFM::PopOnCommitHandler(&SCommitKey);
+		AutoRTFM::AbortTransaction();
+	});
+}
+
+TEST_CASE("Commit.PushOnCommitHandler_WithPopAll_WithAbort")
+{
+	AutoRTFM::Testing::Abort([&]
+	{
+		AutoRTFM::PushOnCommitHandler(&SCommitKey, []() { FAIL("Unreachable!"); });
+		AutoRTFM::PopAllOnCommitHandlers(&SCommitKey);
+		AutoRTFM::AbortTransaction();
+	});
+}
+
+TEST_CASE("Commit.PushOnCommitHandler_Duplicates1")
+{
+	bool bHit = false;
+
+	AutoRTFM::Testing::Commit([&]
+	{
+		AutoRTFM::PushOnCommitHandler(&SCommitKey, [&bHit]() { bHit = 77; });
+		AutoRTFM::PushOnCommitHandler(&SCommitKey, []() { FAIL("Unreachable!"); });
+		AutoRTFM::PopOnCommitHandler(&SCommitKey);
+	});
+
+	// The first push on commit will still go through.
+	REQUIRE(bHit);
+}
+
+TEST_CASE("Commit.PushOnCommitHandler_PopAll_Duplicates")
+{
+	AutoRTFM::Testing::Commit([&]
+	{
+		AutoRTFM::PushOnCommitHandler(&SCommitKey, []() { FAIL("Unreachable!"); });
+		AutoRTFM::PushOnCommitHandler(&SCommitKey, []() { FAIL("Unreachable!"); });
+		AutoRTFM::PopAllOnCommitHandlers(&SCommitKey);
+	});
+}
+
+TEST_CASE("Commit.PushOnCommitHandler_Duplicates2")
+{
+	int Value = 55;
+	
+	AutoRTFM::Testing::Commit([&]
+	{
+		Value = 66;
+		AutoRTFM::PushOnCommitHandler(&SCommitKey, [&Value]() { Value += 13; });
+		AutoRTFM::PushOnCommitHandler(&SCommitKey, [&Value]() { Value *= 11; });
+		Value = 99;
+	});
+
+	REQUIRE(Value == 1232);
+}
+
+TEST_CASE("Commit.PushOnCommitHandler_Order")
+{
+	SECTION("HandlerSandwich")
+	{
+		SECTION("WithoutPop")
+		{
+			int Value = 37;
+			
+			AutoRTFM::Testing::Commit([&]
+			{
+				AutoRTFM::OnCommit([&Value] { REQUIRE(99 == Value); Value += 1; });
+				AutoRTFM::PushOnCommitHandler(&SCommitKey, [&Value]() { REQUIRE(100 == Value); Value += 2; });
+				AutoRTFM::OnCommit([&Value] { REQUIRE(102 == Value); Value += 3; });
+
+				Value = 99;
+			});
+
+			REQUIRE(Value == 105);
+		}
+
+		SECTION("WithPop")
+		{
+			int Value = 37;
+			
+			AutoRTFM::Testing::Commit([&]
+			{
+				AutoRTFM::OnCommit([&Value] { REQUIRE(99 == Value); Value += 1; });
+				AutoRTFM::PushOnCommitHandler(&SCommitKey, [&Value]() { FAIL("Unreachable!"); });
+				AutoRTFM::OnCommit([&Value] { REQUIRE(100 == Value); Value += 3; });
+
+				AutoRTFM::PopOnCommitHandler(&SCommitKey);
+
+				Value = 99;
+			});
+
+			REQUIRE(Value == 103);
+		}
+	}
+
+	SECTION("HandlerInChild")
+	{
+		SECTION("WithoutPop")
+		{
+			int Value = 37;
+			
+			AutoRTFM::Testing::Commit([&]
+			{
+				AutoRTFM::OnCommit([&Value]
+				{
+					REQUIRE(99 == Value);
+					Value += 1;
+				});
+
+				// Make a child transaction.
+				AutoRTFM::Testing::Commit([&]
+				{
+					AutoRTFM::PushOnCommitHandler(&SCommitKey, [&Value]()
+					{
+						REQUIRE(100 == Value);
+						Value += 2;
+					});
+				});
+
+				AutoRTFM::OnCommit([&Value]
+				{
+					REQUIRE(102 == Value);
+					Value += 3;
+				});
+
+				Value = 99;
+			});
+
+			REQUIRE(105 == Value);
+		}
+
+		SECTION("WithPop")
+		{
+			int Value = 37;
+			
+			AutoRTFM::Testing::Commit([&]
+			{
+				AutoRTFM::OnCommit([&Value] { REQUIRE(99 == Value); Value += 1; });
+
+				// Make a child transaction.
+				AutoRTFM::Testing::Commit([&]
+				{
+					AutoRTFM::PushOnCommitHandler(&SCommitKey, [&Value]()
+					{
+						REQUIRE(false);
+					});
+				});
+
+				AutoRTFM::OnCommit([&Value] { REQUIRE(100 == Value); Value += 3; });
+
+				// Bit funky, but we can pop the child's push here!
+				AutoRTFM::PopOnCommitHandler(&SCommitKey);
+
+				Value = 99;
+			});
+
+			REQUIRE(Value == 103);
+		}
+
+		SECTION("AbortInChild")
+		{
+			int Value = 99;
+			
+			AutoRTFM::Testing::Commit([&]
+			{
+				AutoRTFM::OnCommit([&Value] { REQUIRE(37 == Value); Value += 1; });
+
+				// Make a child transaction.
+				AutoRTFM::Testing::Abort([&]
+				{
+					AutoRTFM::PushOnCommitHandler(&SCommitKey, [&Value]() { FAIL("Unreachable!"); });
+					AutoRTFM::AbortTransaction();
+				});
+
+				AutoRTFM::OnCommit([&Value] { REQUIRE(38 == Value); Value += 3; });
+
+				Value = 37;
+			});
+
+			REQUIRE(Value == 41);
+		}
+
+		SECTION("PopInChild")
+		{
+			int Value = 99;
+			
+			AutoRTFM::Testing::Commit([&]
+			{
+				AutoRTFM::PushOnCommitHandler(&SCommitKey, [&Value] { Value = 42; });
+				
+				AutoRTFM::Testing::Commit([&]
+				{
+					AutoRTFM::PopOnCommitHandler(&SCommitKey);
+				});
+			});
+
+			REQUIRE(99 == Value);
+		}
+
+		SECTION("PopInChildAndAbort")
+		{
+			int Value = 99;
+			
+			AutoRTFM::Testing::Commit([&]
+			{
+				AutoRTFM::PushOnCommitHandler(&SCommitKey, [&Value] { Value = 42; });
+				
+				AutoRTFM::Testing::Abort([&]
+				{
+					AutoRTFM::PopOnCommitHandler(&SCommitKey);
+					AutoRTFM::AbortTransaction();
+				});
+			});
+
+			REQUIRE(42 == Value);
+		}
+
+		SECTION("PopInChildsChild")
+		{
+			int Value = 99;
+			
+			AutoRTFM::Testing::Commit([&]
+			{
+				AutoRTFM::PushOnCommitHandler(&SCommitKey, [&Value] { Value = 42; });
+				
+				AutoRTFM::Testing::Commit([&]
+				{
+					AutoRTFM::Testing::Commit([&]
+					{
+						AutoRTFM::PopOnCommitHandler(&SCommitKey);
+					});
+				});
+			});
+
+			REQUIRE(99 == Value);
+		}
+
+		SECTION("PopAllInChild")
+		{
+			int Value = 99;
+			
+			AutoRTFM::Testing::Commit([&]
+			{
+				AutoRTFM::PushOnCommitHandler(&SCommitKey, [&Value] { Value = 42; });
+				AutoRTFM::Testing::Commit([&]
+				{
+					AutoRTFM::PushOnCommitHandler(&SCommitKey, [&Value] { Value = 13; });
+					AutoRTFM::PopAllOnCommitHandlers(&SCommitKey);
+				});
+			});
+
+			REQUIRE(99 == Value);
+		}
+
+		SECTION("PopAllInChildAbort")
+		{
+			int Value = 99;
+			
+			AutoRTFM::Testing::Commit([&]
+			{
+				AutoRTFM::PushOnCommitHandler(&SCommitKey, [&Value] { Value = 42; });
+				
+				AutoRTFM::Testing::Abort([&]
+				{
+					AutoRTFM::PushOnCommitHandler(&SCommitKey, [&Value] { Value = 13; });
+					AutoRTFM::PopAllOnCommitHandlers(&SCommitKey);
+
+					// This abort will ensure that the pop-all cannot affect the outer transactions
+					// push on commit!
+					AutoRTFM::AbortTransaction();
+				});
+			});
+
+			REQUIRE(42 == Value);
+		}
+	}
+}
diff --git a/Engine/Source/Runtime/Core/Private/AutoRTFM/API.cpp b/Engine/Source/Runtime/Core/Private/AutoRTFM/API.cpp
--- a/Engine/Source/Runtime/Core/Private/AutoRTFM/API.cpp
+++ b/Engine/Source/Runtime/Core/Private/AutoRTFM/API.cpp
@@ -478,6 +478,27 @@ AUTORTFM_CORE_API void RTFM_OnAbortInternal(TFunction<void()>&& Work)
 	Context->GetCurrentTransaction()->DeferUntilAbort(MoveTemp(Work));
 }
 
+AUTORTFM_CORE_API void RTFM_PushOnCommitHandlerInternal(const void* Key, TFunction<void()>&& Work)
+{
+	FContext* Context = FContext::Get();
+	ASSERT(Context->GetStatus() == EContextStatus::OnTrack);
+	Context->GetCurrentTransaction()->PushDeferUntilCommitHandler(Key, MoveTemp(Work));
+}
+
+AUTORTFM_CORE_API void RTFM_PopOnCommitHandlerInternal(const void* Key)
+{
+	FContext* Context = FContext::Get();
+	ASSERT(Context->GetStatus() == EContextStatus::OnTrack);
+	Context->GetCurrentTransaction()->PopDeferUntilCommitHandler(Key);
+}
+
+AUTORTFM_CORE_API void RTFM_PopAllOnCommitHandlersInternal(const void* Key)
+{
+	FContext* Context = FContext::Get();
+	ASSERT(Context->GetStatus() == EContextStatus::OnTrack);
+	Context->GetCurrentTransaction()->PopAllDeferUntilCommitHandlers(Key);
+}
+
 AUTORTFM_CORE_API void RTFM_PushOnAbortHandlerInternal(const void* Key, TFunction<void()>&& Work)
 {
 	FContext* Context = FContext::Get();
@@ -815,6 +836,30 @@ void ForTheRuntime::OnAbortInternal(TFunction<void()> && Work)
 	}
 }
 
+void ForTheRuntime::PushOnCommitHandlerInternal(const void* Key, TFunction<void()>&& Work)
+{
+	if (autortfm_is_closed())
+	{
+		return RTFM_PushOnCommitHandlerInternal(Key, MoveTemp(Work));
+	}
+}
+
+void ForTheRuntime::PopOnCommitHandlerInternal(const void* Key)
+{
+	if (autortfm_is_closed())
+	{
+		return RTFM_PopOnCommitHandlerInternal(Key);
+	}
+}
+
+void ForTheRuntime::PopAllOnCommitHandlersInternal(const void* Key)
+{
+	if (autortfm_is_closed())
+	{
+		return RTFM_PopAllOnCommitHandlersInternal(Key);
+	}
+}
+
 void ForTheRuntime::PushOnAbortHandlerInternal(const void* Key, TFunction<void()> && Work)
 {
 	if (autortfm_is_closed())
diff --git a/Engine/Source/Runtime/Core/Private/AutoRTFM/TaskArray.h b/Engine/Source/Runtime/Core/Private/AutoRTFM/TaskArray.h
--- a/Engine/Source/Runtime/Core/Private/AutoRTFM/TaskArray.h
+++ b/Engine/Source/Runtime/Core/Private/AutoRTFM/TaskArray.h
@@ -317,8 +317,22 @@ private:
 		{
 			BitMaskPerKeyType* const KeyedElement = BitMasksPerKey.Find(*Pair.Key);
 
-			// We don't support keyed elements in the forward iterator (yet?).
-			ASSERT(!KeyedElement);
+			if (0 < KeyedElement->Num())
+			{
+				const bool bFirst = (*KeyedElement)[0];
+				KeyedElement->RemoveAt(0);
+
+				if (bFirst)
+				{
+					static TFunction<void()> EmptyFunctor = [] {};
+
+					// We have a deferred signal to pop this element, so do it now.
+					*Pair.Key = SKey{};
+					Pair.Val->Reset();
+					*Pair.Val = EmptyFunctor;
+					return;
+				}
+			}
 		}
 
 		Func(*Pair.Val);
diff --git a/Engine/Source/Runtime/Core/Private/AutoRTFM/Transaction.cpp b/Engine/Source/Runtime/Core/Private/AutoRTFM/Transaction.cpp
--- a/Engine/Source/Runtime/Core/Private/AutoRTFM/Transaction.cpp
+++ b/Engine/Source/Runtime/Core/Private/AutoRTFM/Transaction.cpp
@@ -63,6 +63,11 @@ void FTransaction::Reset()
 	// Reset to the initial state.
 	CurrentState = EState::Uninitialized;
 
+	DeferredPopOnCommitHandlers.Reset();
+	DeferredPopOnAbortHandlers.Reset();
+	DeferredPopAllOnCommitHandlers.Reset();
+	DeferredPopAllOnAbortHandlers.Reset();
+
 	ASSERT(IsFresh());
 }
 
@@ -73,7 +78,11 @@ bool FTransaction::IsFresh() const
         && WriteLog.IsEmpty()
         && CommitTasks.IsEmpty()
         && AbortTasks.IsEmpty()
-        && !IsDone();
+        && !IsDone()
+		&& DeferredPopOnCommitHandlers.IsEmpty()
+		&& DeferredPopOnAbortHandlers.IsEmpty()
+		&& DeferredPopAllOnCommitHandlers.IsEmpty()
+		&& DeferredPopAllOnAbortHandlers.IsEmpty();
 }
 
 void FTransaction::AbortWithoutThrowing()
@@ -224,6 +233,40 @@ void FTransaction::CommitNested()
 
     Parent->WriteLogBumpAllocator.Merge(MoveTemp(WriteLogBumpAllocator));
 
+	// For all the deferred calls to `PopOnCommitHandler` that we couldn't
+	// process (because our transaction nest didn't `PushOnCommitHandler`)
+	// we need to move these to the parent now to handle them.
+	for (const void* Key : DeferredPopOnCommitHandlers)
+	{
+		Parent->PopDeferUntilCommitHandler(Key);
+	}
+	DeferredPopOnCommitHandlers.Reset();
+
+	// For all the deferred calls to `PopOnAbortHandler` that we couldn't
+	// process (because our transaction nest didn't `PushOnAbortHandler`)
+	// we need to move these to the parent now to handle them.
+	for (const void* Key : DeferredPopOnAbortHandlers)
+	{
+		Parent->PopDeferUntilAbortHandler(Key);
+	}
+	DeferredPopOnAbortHandlers.Reset();
+
+	// For all the calls to `PopAllOnCommitHandlers` we need to run these
+	// again on parent now to handle them there too.
+	for (const void* Key : DeferredPopAllOnCommitHandlers)
+	{
+		Parent->PopAllDeferUntilCommitHandlers(Key);
+	}
+	DeferredPopAllOnCommitHandlers.Reset();
+
+	// For all the calls to `PopAllOnAbortHandlers` we need to run these
+	// again on parent now to handle them there too.
+	for (const void* Key : DeferredPopAllOnAbortHandlers)
+	{
+		Parent->PopAllDeferUntilAbortHandlers(Key);
+	}
+	DeferredPopAllOnAbortHandlers.Reset();
+
     Parent->CommitTasks.AddAll(MoveTemp(CommitTasks));
     Parent->AbortTasks.AddAll(MoveTemp(AbortTasks));
 
diff --git a/Engine/Source/Runtime/Core/Private/AutoRTFM/Transaction.h b/Engine/Source/Runtime/Core/Private/AutoRTFM/Transaction.h
--- a/Engine/Source/Runtime/Core/Private/AutoRTFM/Transaction.h
+++ b/Engine/Source/Runtime/Core/Private/AutoRTFM/Transaction.h
@@ -93,9 +93,12 @@ public:
 
     void DeferUntilCommit(TFunction<void()>&&);
     void DeferUntilAbort(TFunction<void()>&&);
+	void PushDeferUntilCommitHandler(const void* Key, TFunction<void()>&&);
+	void PopDeferUntilCommitHandler(const void* Key);
+	void PopAllDeferUntilCommitHandlers(const void* Key);
     void PushDeferUntilAbortHandler(const void* Key, TFunction<void()>&&);
-	bool PopDeferUntilAbortHandler(const void* Key);
-	bool PopAllDeferUntilAbortHandlers(const void* Key);
+	void PopDeferUntilAbortHandler(const void* Key);
+	void PopAllDeferUntilAbortHandlers(const void* Key);
 
     [[noreturn]] void AbortAndThrow();
     void AbortWithoutThrowing();
@@ -178,10 +181,28 @@ private:
     // Are we nested? Then this is the parent.
     FTransaction* Parent{nullptr};
 
-    // Commit tasks run on commit in forward order. Abort tasks run on abort in reverse order.
+    // Commit tasks run on commit in forward order.
     TTaskArray<TFunction<void()>> CommitTasks;
+
+	// Abort tasks run on abort in reverse order.
     TTaskArray<TFunction<void()>> AbortTasks;
 
+	// If a call to `PopOnCommitHandler` could not find a commit to pop, it is deferred
+	// and tried again on the parent transaction.
+	TArray<const void*, TInlineAllocator<8>> DeferredPopOnCommitHandlers;
+
+	// If a call to `PopOnAbortHandler` could not find an abort to pop, it is deferred
+	// and tried again on the parent transaction.
+	TArray<const void*, TInlineAllocator<8>> DeferredPopOnAbortHandlers;
+
+	// If a call to `PopAllOnCommitHandlers` was used and our transaction successfully
+	// commits, we need to propagate this to the parent too.
+	TArray<const void*, TInlineAllocator<1>> DeferredPopAllOnCommitHandlers;
+
+	// If a call to `PopAllOnAbortHandlers` was used and our transaction successfully
+	// commits, we need to propagate this to the parent too.
+	TArray<const void*, TInlineAllocator<1>> DeferredPopAllOnAbortHandlers;
+
     FHitSet HitSet;
     FIntervalTree NewMemoryTracker;
     FWriteLog WriteLog;
diff --git a/Engine/Source/Runtime/Core/Private/AutoRTFM/TransactionInlines.h b/Engine/Source/Runtime/Core/Private/AutoRTFM/TransactionInlines.h
--- a/Engine/Source/Runtime/Core/Private/AutoRTFM/TransactionInlines.h
+++ b/Engine/Source/Runtime/Core/Private/AutoRTFM/TransactionInlines.h
@@ -249,6 +249,33 @@ UE_AUTORTFM_FORCEINLINE void FTransaction::DeferUntilAbort(TFunction<void()>&& C
     AbortTasks.Add(MoveTemp(Copy));
 }
 
+UE_AUTORTFM_FORCEINLINE void FTransaction::PushDeferUntilCommitHandler(const void* Key, TFunction<void()>&& Callback)
+{
+	// We explicitly must copy the function here because the original was allocated
+	// within a transactional context, and thus the memory was allocated under
+	// transactionalized conditions. By copying, we create an open copy of the callback.
+	TFunction<void()> Copy(Callback);
+	CommitTasks.AddKeyed(Key, MoveTemp(Copy));
+}
+
+UE_AUTORTFM_FORCEINLINE void FTransaction::PopDeferUntilCommitHandler(const void* Key)
+{
+	if (LIKELY(CommitTasks.DeleteKey(Key)))
+	{
+		return;
+	}
+
+	DeferredPopOnCommitHandlers.Push(Key);
+}
+
+UE_AUTORTFM_FORCEINLINE void FTransaction::PopAllDeferUntilCommitHandlers(const void* Key)
+{
+	CommitTasks.DeleteAllMatchingKeys(Key);
+
+	// We also need to remember to run this on our parent's nest if our transaction commits.
+	DeferredPopAllOnCommitHandlers.Push(Key);
+}
+
 UE_AUTORTFM_FORCEINLINE void FTransaction::PushDeferUntilAbortHandler(const void* Key, TFunction<void()>&& Callback)
 {
 	// We explicitly must copy the function here because the original was allocated
@@ -258,14 +285,22 @@ UE_AUTORTFM_FORCEINLINE void FTransaction::PushDeferUntilAbortHandler(const void
     AbortTasks.AddKeyed(Key, MoveTemp(Copy));
 }
 
-UE_AUTORTFM_FORCEINLINE bool FTransaction::PopDeferUntilAbortHandler(const void* Key)
+UE_AUTORTFM_FORCEINLINE void FTransaction::PopDeferUntilAbortHandler(const void* Key)
 {
-	return AbortTasks.DeleteKey(Key);
+	if (LIKELY(AbortTasks.DeleteKey(Key)))
+	{
+		return;
+	}
+
+	DeferredPopOnAbortHandlers.Push(Key);
 }
 
-UE_AUTORTFM_FORCEINLINE bool FTransaction::PopAllDeferUntilAbortHandlers(const void* Key)
+UE_AUTORTFM_FORCEINLINE void FTransaction::PopAllDeferUntilAbortHandlers(const void* Key)
 {
-	return AbortTasks.DeleteAllMatchingKeys(Key);
+	AbortTasks.DeleteAllMatchingKeys(Key);
+
+	// We also need to remember to run this on our parent's nest if our transaction commits.
+	DeferredPopAllOnAbortHandlers.Push(Key);
 }
 
 UE_AUTORTFM_FORCEINLINE void FTransaction::CollectStats() const
diff --git a/Engine/Source/Runtime/Core/Public/AutoRTFM/AutoRTFM.h b/Engine/Source/Runtime/Core/Public/AutoRTFM/AutoRTFM.h
--- a/Engine/Source/Runtime/Core/Public/AutoRTFM/AutoRTFM.h
+++ b/Engine/Source/Runtime/Core/Public/AutoRTFM/AutoRTFM.h
@@ -488,6 +488,9 @@ namespace ForTheRuntime
 	UE_AUTORTFM_API void CascadingRetryTransactionInternal(TFunction<void()>&& RunAfterAbortBeforeRetryWork);
 	UE_AUTORTFM_API void OnCommitInternal(TFunction<void()>&& Work);
 	UE_AUTORTFM_API void OnAbortInternal(TFunction<void()>&& Work);
+	UE_AUTORTFM_API void PushOnCommitHandlerInternal(const void* Key, TFunction<void()>&& Work);
+	UE_AUTORTFM_API void PopOnCommitHandlerInternal(const void* Key);
+	UE_AUTORTFM_API void PopAllOnCommitHandlersInternal(const void* Key);
 	UE_AUTORTFM_API void PushOnAbortHandlerInternal(const void* Key, TFunction<void()>&& Work);
 	UE_AUTORTFM_API void PopOnAbortHandlerInternal(const void* Key);
 	UE_AUTORTFM_API void PopAllOnAbortHandlersInternal(const void* Key);
@@ -871,9 +874,57 @@ template<typename FunctorType> static UE_AUTORTFM_FORCEINLINE void OnAbort(Funct
 template<typename FunctorType> static UE_AUTORTFM_FORCEINLINE void OnAbort(FunctorType&&) {}
 #endif
 
+#if UE_AUTORTFM
+// Register a handler for transaction commit. Takes a key parameter so that
+// the handler can be unregistered (see `PopOnCommitHandler`). This is useful
+// for scoped mutations that need an abort handler present unless execution
+// reaches the end of the relevant scope.
+template<typename FunctorType> static UE_AUTORTFM_FORCEINLINE void PushOnCommitHandler(const void* Key, FunctorType&& Work)
+{
+	if (autortfm_is_closed())
+	{
+		ForTheRuntime::PushOnCommitHandlerInternal(Key, std::forward<FunctorType>(Work));
+	}
+}
+#else
+// Register a handler for transaction commit. Takes a key parameter so that
+// the handler can be unregistered (see `PopOnCommitHandler`). This is useful
+// for scoped mutations that need an abort handler present unless execution
+// reaches the end of the relevant scope.
+template<typename FunctorType> static UE_AUTORTFM_FORCEINLINE void PushOnCommitHandler(const void*, FunctorType&&) {}
+#endif
+
+#if UE_AUTORTFM
+// Unregister the most recently pushed handler (via `PushOnCommitHandler`) for the given key.
+static UE_AUTORTFM_FORCEINLINE void PopOnCommitHandler(const void* Key)
+{
+	if (autortfm_is_closed())
+	{
+		ForTheRuntime::PopOnCommitHandlerInternal(Key);
+	}
+}
+#else
+// Unregister the most recently pushed handler (via `PushOnCommitHandler`) for the given key.
+static UE_AUTORTFM_FORCEINLINE void PopOnCommitHandler(const void*) {}
+#endif
+
+#if UE_AUTORTFM
+// Unregister all pushed handlers (via `PushOnCommitHandler`) for the given key.
+static UE_AUTORTFM_FORCEINLINE void PopAllOnCommitHandlers(const void* Key)
+{
+	if (autortfm_is_closed())
+	{
+		ForTheRuntime::PopAllOnCommitHandlersInternal(Key);
+	}
+}
+#else
+// Unregister all pushed handlers (via `PushOnCommitHandler`) for the given key.
+static UE_AUTORTFM_FORCEINLINE void PopAllOnCommitHandlers(const void*) {}
+#endif
+
 #if UE_AUTORTFM
 // Register a handler for transaction abort. Takes a key parameter so that
-// the handler can be unregistered (see PopOnAbortHandler). This is useful
+// the handler can be unregistered (see `PopOnAbortHandler`). This is useful
 // for scoped mutations that need an abort handler present unless execution
 // reaches the end of the relevant scope.
 template<typename FunctorType> static UE_AUTORTFM_FORCEINLINE void PushOnAbortHandler(const void* Key, FunctorType&& Work)
@@ -885,14 +936,14 @@ template<typename FunctorType> static UE_AUTORTFM_FORCEINLINE void PushOnAbortHa
 }
 #else
 // Register a handler for transaction abort. Takes a key parameter so that
-// the handler can be unregistered (see PopOnAbortHandler). This is useful
+// the handler can be unregistered (see `PopOnAbortHandler`). This is useful
 // for scoped mutations that need an abort handler present unless execution
 // reaches the end of the relevant scope.
 template<typename FunctorType> static UE_AUTORTFM_FORCEINLINE void PushOnAbortHandler(const void* Key, FunctorType&&) {}
 #endif
 
 #if UE_AUTORTFM
-// Unregister the most recently pushed handler (via PushOnAbortHandler) for the given key.
+// Unregister the most recently pushed handler (via `PushOnAbortHandler`) for the given key.
 static UE_AUTORTFM_FORCEINLINE void PopOnAbortHandler(const void* Key)
 {
 	if (autortfm_is_closed())
@@ -901,12 +952,12 @@ static UE_AUTORTFM_FORCEINLINE void PopOnAbortHandler(const void* Key)
 	}
 }
 #else
-// Unregister the most recently pushed handler (via PushOnAbortHandler) for the given key.
+// Unregister the most recently pushed handler (via `PushOnAbortHandler`) for the given key.
 static UE_AUTORTFM_FORCEINLINE void PopOnAbortHandler(const void* Key) {}
 #endif
 
 #if UE_AUTORTFM
-// Unregister all pushed handlers (via PushOnAbortHandler) for the given key.
+// Unregister all pushed handlers (via `PushOnAbortHandler`) for the given key.
 static UE_AUTORTFM_FORCEINLINE void PopAllOnAbortHandlers(const void* Key)
 {
 	if (autortfm_is_closed())
@@ -915,7 +966,7 @@ static UE_AUTORTFM_FORCEINLINE void PopAllOnAbortHandlers(const void* Key)
 	}
 }
 #else
-// Unregister all pushed handlers (via PushOnAbortHandler) for the given key.
+// Unregister all pushed handlers (via `PushOnAbortHandler`) for the given key.
 static UE_AUTORTFM_FORCEINLINE void PopAllOnAbortHandlers(const void* Key) {}
 #endif
 
