[AutoRTFM] Make FCoreRedirects::[Add|Remove]RedirectList transactionally safe

diff --git a/Engine/Source/Programs/AutoRTFMTests/Private/UECoreTests.cpp b/Engine/Source/Programs/AutoRTFMTests/Private/UECoreTests.cpp
--- a/Engine/Source/Programs/AutoRTFMTests/Private/UECoreTests.cpp
+++ b/Engine/Source/Programs/AutoRTFMTests/Private/UECoreTests.cpp
@@ -1,5 +1,6 @@
 // Copyright Epic Games, Inc. All Rights Reserved.
 
+#include "AutoRTFM/AutoRTFM.h"
 #include "Catch2Includes.h"
 #include "AssetRegistry/AssetDataTagMap.h"
 #include "AutoRTFM/AutoRTFMTesting.h"
@@ -12,6 +13,7 @@
 #include "Logging/LogMacros.h"
 #include "Logging/StructuredLog.h"
 #include "Serialization/CustomVersion.h"
+#include "UObject/CoreRedirects.h"
 #include "UObject/NameTypes.h"
 #include "UObject/UObjectArray.h"
 #include "UObject/UObjectGlobals.h"
@@ -2368,6 +2370,7 @@ TEST_CASE("UECore.UE_LOGFMT")
 		});
 	}
 }
+
 TEST_CASE("UECore.AsyncLoading")
 {
 	SECTION("LoadPackageAsync")
@@ -2540,3 +2543,72 @@ TEST_CASE("UECore.AsyncLoading")
 		});
 	}
 }
+
+TEST_CASE("UECore.CoreRedirects")
+{
+	FCoreRedirects::Initialize();
+
+	FCoreRedirectObjectName From(NAME_None, NAME_None, TEXT("/A/B/C"));
+	FCoreRedirectObjectName To(NAME_None, NAME_None, TEXT("/X/Y/Z"));
+
+	TArray<FCoreRedirect> Redirects;
+	Redirects.Emplace(ECoreRedirectFlags::Type_Package, From, To);
+
+	SECTION("Basic Assumptions")
+	{
+		FCoreRedirects::AddRedirectList(Redirects, TEXT("AutoRTFMTests.UECore.CoreRedirects"));
+		REQUIRE(FCoreRedirects::GetRedirectedName(ECoreRedirectFlags::Type_Package, From) == To);
+		FCoreRedirects::RemoveRedirectList(Redirects, TEXT("AutoRTFMTests.UECore.CoreRedirects"));
+		REQUIRE(FCoreRedirects::GetRedirectedName(ECoreRedirectFlags::Type_Package, From) == From);
+	}
+
+	SECTION("AddRedirectList")
+	{
+		SECTION("Commit")
+		{
+			AutoRTFM::Testing::Commit([&]
+			{
+				FCoreRedirects::AddRedirectList(Redirects, TEXT("AutoRTFMTests.UECore.CoreRedirects"));
+			});
+			REQUIRE(FCoreRedirects::GetRedirectedName(ECoreRedirectFlags::Type_Package, From) == To);
+			FCoreRedirects::RemoveRedirectList(Redirects, TEXT("AutoRTFMTests.UECore.CoreRedirects"));
+			REQUIRE(FCoreRedirects::GetRedirectedName(ECoreRedirectFlags::Type_Package, From) == From);
+		}
+		
+		SECTION("Abort")
+		{
+			AutoRTFM::Testing::Abort([&]
+			{
+				FCoreRedirects::AddRedirectList(Redirects, TEXT("AutoRTFMTests.UECore.CoreRedirects"));
+				AutoRTFM::AbortTransaction();
+			});
+			REQUIRE(FCoreRedirects::GetRedirectedName(ECoreRedirectFlags::Type_Package, From) == From);
+		}
+	}
+	
+	SECTION("RemoveRedirectList")
+	{
+		FCoreRedirects::AddRedirectList(Redirects, TEXT("AutoRTFMTests.UECore.CoreRedirects"));
+		
+		SECTION("Commit")
+		{
+			AutoRTFM::Testing::Commit([&]
+			{
+				FCoreRedirects::RemoveRedirectList(Redirects, TEXT("AutoRTFMTests.UECore.CoreRedirects"));
+			});
+			REQUIRE(FCoreRedirects::GetRedirectedName(ECoreRedirectFlags::Type_Package, From) == From);
+		}
+		
+		SECTION("Abort")
+		{
+			AutoRTFM::Testing::Abort([&]
+			{
+				FCoreRedirects::RemoveRedirectList(Redirects, TEXT("AutoRTFMTests.UECore.CoreRedirects"));
+				AutoRTFM::AbortTransaction();
+			});
+			REQUIRE(FCoreRedirects::GetRedirectedName(ECoreRedirectFlags::Type_Package, From) == To);
+			FCoreRedirects::RemoveRedirectList(Redirects, TEXT("AutoRTFMTests.UECore.CoreRedirects"));
+			REQUIRE(FCoreRedirects::GetRedirectedName(ECoreRedirectFlags::Type_Package, From) == From);
+		}
+	}
+}
diff --git a/Engine/Source/Runtime/CoreUObject/Private/UObject/CoreRedirects.cpp b/Engine/Source/Runtime/CoreUObject/Private/UObject/CoreRedirects.cpp
--- a/Engine/Source/Runtime/CoreUObject/Private/UObject/CoreRedirects.cpp
+++ b/Engine/Source/Runtime/CoreUObject/Private/UObject/CoreRedirects.cpp
@@ -2218,6 +2218,18 @@ bool FCoreRedirects::ReadRedirectsFromIni(const FString& IniName)
 
 bool FCoreRedirects::AddRedirectList(TArrayView<const FCoreRedirect> Redirects, const FString& SourceString)
 {
+	if (AutoRTFM::IsClosed())
+	{
+		UE_AUTORTFM_ONABORT(Redirects, SourceString)
+		{
+			RemoveRedirectList(Redirects, SourceString);
+		};
+		return AutoRTFM::Open([&Redirects, &SourceString]
+		{
+			return AddRedirectList(Redirects, SourceString);
+		});
+	}
+
 	FCoreRedirectorScopeLockForWrite ScopeLock(RWLock);
 	ensureMsgf(IsInitialized(), TEXT("FCoreRedirects must be initialized on the game thread before use."));
 	return AddRedirectListUnderWriteLock(Redirects, SourceString, ScopeLock);
@@ -2383,6 +2395,18 @@ bool FCoreRedirects::AddSingleRedirectUnderWriteLock(const FCoreRedirect& NewRed
 
 bool FCoreRedirects::RemoveRedirectList(TArrayView<const FCoreRedirect> Redirects, const FString& SourceString)
 {
+	if (AutoRTFM::IsClosed())
+	{
+		UE_AUTORTFM_ONABORT(Redirects, SourceString)
+		{
+			AddRedirectList(Redirects, SourceString);
+		};
+		return AutoRTFM::Open([&Redirects, &SourceString]
+		{
+			return RemoveRedirectList(Redirects, SourceString);
+		});
+	}
+
 	FCoreRedirectorScopeLockForWrite ScopeLock(RWLock);
 	UE_LOG(LogCoreRedirects, Verbose, TEXT("RemoveRedirect(%s) Removing %d redirects"), *SourceString, Redirects.Num());
 
