Fix that you could not debug step into `lookup_and_call` functions in closed code.

diff --git a/Engine/Source/Programs/AutoRTFMTests/Private/CrossCUTests.cpp b/Engine/Source/Programs/AutoRTFMTests/Private/CrossCUTests.cpp
new file mode 100644
--- /dev/null
+++ b/Engine/Source/Programs/AutoRTFMTests/Private/CrossCUTests.cpp
@@ -0,0 +1,25 @@
+// Copyright Epic Games, Inc. All Rights Reserved.
+
+#include "API.h"
+#include "Catch2Includes.h"
+#include "CoreTypes.h"
+
+#include <AutoRTFM/AutoRTFMTesting.h>
+
+UE_DISABLE_OPTIMIZATION_SHIP
+
+namespace CrossCU
+{
+int SomeFunction(int X);
+} // CrossCU
+
+TEST_CASE("CrossCU.Call")
+{
+    AutoRTFM::Testing::Commit([&]
+	{
+		int Value = CrossCU::SomeFunction(0);
+		REQUIRE(Value == 42);
+	});
+}
+
+UE_ENABLE_OPTIMIZATION_SHIP
diff --git a/Engine/Source/Programs/AutoRTFMTests/Private/CrossCUTestsOther.cpp b/Engine/Source/Programs/AutoRTFMTests/Private/CrossCUTestsOther.cpp
new file mode 100644
--- /dev/null
+++ b/Engine/Source/Programs/AutoRTFMTests/Private/CrossCUTestsOther.cpp
@@ -0,0 +1,16 @@
+// Copyright Epic Games, Inc. All Rights Reserved.
+
+#include "API.h"
+#include "Catch2Includes.h"
+#include "CoreTypes.h"
+
+#include <AutoRTFM/AutoRTFMTesting.h>
+
+UE_DISABLE_OPTIMIZATION_SHIP
+
+namespace CrossCU
+{
+int SomeFunction(int X) { return X + 42; }
+} // CrossCU
+
+UE_ENABLE_OPTIMIZATION_SHIP
