[AutoRTFM] Replace FTransactionallySafeRWScopeLock with UE::TRWScopeLock.

FTransactionallySafeRWScopeLock is unnecessary, as it is functionally identical to a UE::TRWScopeLock.

diff --git a/Engine/Source/Programs/AutoRTFMTests/Private/UECoreTests.cpp b/Engine/Source/Programs/AutoRTFMTests/Private/UECoreTests.cpp
--- a/Engine/Source/Programs/AutoRTFMTests/Private/UECoreTests.cpp
+++ b/Engine/Source/Programs/AutoRTFMTests/Private/UECoreTests.cpp
@@ -19,6 +19,7 @@
 #include "UObject/UObjectGlobals.h"
 #include "MyAutoRTFMTestObject.h"
 #include "Containers/Queue.h"
+#include "Misc/ScopeRWLock.h"
 #include "Misc/TransactionallySafeRWScopeLock.h"
 #include "Misc/TransactionallySafeMutex.h"
 #include "Misc/ConfigCacheIni.h"
@@ -1152,7 +1153,7 @@ TEST_CASE("UECore.TransactionallySafeRWScopeLock")
 
 		AutoRTFM::ETransactionResult Result = AutoRTFM::Transact([&]
 			{
-				FTransactionallySafeRWScopeLock Lock(RWLock, FRWScopeLockType::SLT_ReadOnly);
+				UE::TRWScopeLock Lock(RWLock, FRWScopeLockType::SLT_ReadOnly);
 				AutoRTFM::AbortTransaction();
 			});
 
@@ -1160,7 +1161,7 @@ TEST_CASE("UECore.TransactionallySafeRWScopeLock")
 
 		Result = AutoRTFM::Transact([&]
 			{
-				FTransactionallySafeRWScopeLock Lock(RWLock, FRWScopeLockType::SLT_ReadOnly);
+				UE::TRWScopeLock Lock(RWLock, FRWScopeLockType::SLT_ReadOnly);
 			});
 
 		REQUIRE(AutoRTFM::ETransactionResult::Committed == Result);
@@ -1171,7 +1172,7 @@ TEST_CASE("UECore.TransactionallySafeRWScopeLock")
 		AutoRTFM::ETransactionResult Result = AutoRTFM::Transact([&]
 			{
 				FTransactionallySafeRWLock RWLock;
-				FTransactionallySafeRWScopeLock Lock(RWLock, FRWScopeLockType::SLT_ReadOnly);
+				UE::TRWScopeLock Lock(RWLock, FRWScopeLockType::SLT_ReadOnly);
 				AutoRTFM::AbortTransaction();
 			});
 
@@ -1180,7 +1181,7 @@ TEST_CASE("UECore.TransactionallySafeRWScopeLock")
 		Result = AutoRTFM::Transact([&]
 			{
 				FTransactionallySafeRWLock RWLock;
-				FTransactionallySafeRWScopeLock Lock(RWLock, FRWScopeLockType::SLT_ReadOnly);
+				UE::TRWScopeLock Lock(RWLock, FRWScopeLockType::SLT_ReadOnly);
 			});
 
 		REQUIRE(AutoRTFM::ETransactionResult::Committed == Result);
@@ -1196,7 +1197,7 @@ TEST_CASE("UECore.TransactionallySafeRWScopeLock")
 
 				InnerResult = AutoRTFM::Transact([&]
 					{
-						FTransactionallySafeRWScopeLock Lock(RWLock, FRWScopeLockType::SLT_ReadOnly);
+						UE::TRWScopeLock Lock(RWLock, FRWScopeLockType::SLT_ReadOnly);
 						AutoRTFM::CascadingAbortTransaction();
 					});
 			});
@@ -1209,7 +1210,7 @@ TEST_CASE("UECore.TransactionallySafeRWScopeLock")
 
 				InnerResult = AutoRTFM::Transact([&]
 					{
-						FTransactionallySafeRWScopeLock Lock(RWLock, FRWScopeLockType::SLT_ReadOnly);
+						UE::TRWScopeLock Lock(RWLock, FRWScopeLockType::SLT_ReadOnly);
 						AutoRTFM::AbortTransaction();
 					});
 			});
@@ -1223,7 +1224,7 @@ TEST_CASE("UECore.TransactionallySafeRWScopeLock")
 
 				InnerResult = AutoRTFM::Transact([&]
 					{
-						FTransactionallySafeRWScopeLock Lock(RWLock, FRWScopeLockType::SLT_ReadOnly);
+						UE::TRWScopeLock Lock(RWLock, FRWScopeLockType::SLT_ReadOnly);
 					});
 
 				AutoRTFM::AbortTransaction();
@@ -1237,7 +1238,7 @@ TEST_CASE("UECore.TransactionallySafeRWScopeLock")
 
 				InnerResult = AutoRTFM::Transact([&]
 					{
-						FTransactionallySafeRWScopeLock Lock(RWLock, FRWScopeLockType::SLT_ReadOnly);
+						UE::TRWScopeLock Lock(RWLock, FRWScopeLockType::SLT_ReadOnly);
 					});
 			});
 
@@ -1251,7 +1252,7 @@ TEST_CASE("UECore.TransactionallySafeRWScopeLock")
 
 		AutoRTFM::ETransactionResult Result = AutoRTFM::Transact([&]
 			{
-				FTransactionallySafeRWScopeLock Lock(RWLock, FRWScopeLockType::SLT_Write);
+				UE::TRWScopeLock Lock(RWLock, FRWScopeLockType::SLT_Write);
 				AutoRTFM::AbortTransaction();
 			});
 
@@ -1259,7 +1260,7 @@ TEST_CASE("UECore.TransactionallySafeRWScopeLock")
 
 		Result = AutoRTFM::Transact([&]
 			{
-				FTransactionallySafeRWScopeLock Lock(RWLock, FRWScopeLockType::SLT_Write);
+				UE::TRWScopeLock Lock(RWLock, FRWScopeLockType::SLT_Write);
 			});
 
 		REQUIRE(AutoRTFM::ETransactionResult::Committed == Result);
@@ -1270,7 +1271,7 @@ TEST_CASE("UECore.TransactionallySafeRWScopeLock")
 		AutoRTFM::ETransactionResult Result = AutoRTFM::Transact([&]
 			{
 				FTransactionallySafeRWLock RWLock;
-				FTransactionallySafeRWScopeLock Lock(RWLock, FRWScopeLockType::SLT_Write);
+				UE::TRWScopeLock Lock(RWLock, FRWScopeLockType::SLT_Write);
 				AutoRTFM::AbortTransaction();
 			});
 
@@ -1279,7 +1280,7 @@ TEST_CASE("UECore.TransactionallySafeRWScopeLock")
 		Result = AutoRTFM::Transact([&]
 			{
 				FTransactionallySafeRWLock RWLock;
-				FTransactionallySafeRWScopeLock Lock(RWLock, FRWScopeLockType::SLT_Write);
+				UE::TRWScopeLock Lock(RWLock, FRWScopeLockType::SLT_Write);
 			});
 
 		REQUIRE(AutoRTFM::ETransactionResult::Committed == Result);
@@ -1295,7 +1296,7 @@ TEST_CASE("UECore.TransactionallySafeRWScopeLock")
 
 				InnerResult = AutoRTFM::Transact([&]
 					{
-						FTransactionallySafeRWScopeLock Lock(RWLock, FRWScopeLockType::SLT_Write);
+						UE::TRWScopeLock Lock(RWLock, FRWScopeLockType::SLT_Write);
 						AutoRTFM::CascadingAbortTransaction();
 					});
 			});
@@ -1308,7 +1309,7 @@ TEST_CASE("UECore.TransactionallySafeRWScopeLock")
 
 				InnerResult = AutoRTFM::Transact([&]
 					{
-						FTransactionallySafeRWScopeLock Lock(RWLock, FRWScopeLockType::SLT_Write);
+						UE::TRWScopeLock Lock(RWLock, FRWScopeLockType::SLT_Write);
 						AutoRTFM::AbortTransaction();
 					});
 			});
@@ -1322,7 +1323,7 @@ TEST_CASE("UECore.TransactionallySafeRWScopeLock")
 
 				InnerResult = AutoRTFM::Transact([&]
 					{
-						FTransactionallySafeRWScopeLock Lock(RWLock, FRWScopeLockType::SLT_Write);
+						UE::TRWScopeLock Lock(RWLock, FRWScopeLockType::SLT_Write);
 					});
 
 				AutoRTFM::AbortTransaction();
@@ -1336,7 +1337,7 @@ TEST_CASE("UECore.TransactionallySafeRWScopeLock")
 
 				InnerResult = AutoRTFM::Transact([&]
 					{
-						FTransactionallySafeRWScopeLock Lock(RWLock, FRWScopeLockType::SLT_Write);
+						UE::TRWScopeLock Lock(RWLock, FRWScopeLockType::SLT_Write);
 					});
 			});
 
diff --git a/Engine/Source/Runtime/AssetRegistry/Private/AssetRegistry.cpp b/Engine/Source/Runtime/AssetRegistry/Private/AssetRegistry.cpp
--- a/Engine/Source/Runtime/AssetRegistry/Private/AssetRegistry.cpp
+++ b/Engine/Source/Runtime/AssetRegistry/Private/AssetRegistry.cpp
@@ -33,7 +33,7 @@
 #include "Misc/ScopeExit.h"
 #include "Misc/ScopeRWLock.h"
 #include "Misc/TrackedActivity.h"
-#include "Misc/TransactionallySafeRWScopeLock.h"
+#include "Misc/TransactionallySafeRWLock.h"
 #include "AssetRegistry/PackageReader.h"
 #include "ProfilingDebugging/MiscTrace.h"
 #include "Serialization/ArrayReader.h"
@@ -286,7 +286,7 @@ public:
 	}
 
 	UE::AssetRegistry::Private::FRWLockWithPriority& Lock;
-	TOptional<FTransactionallySafeRWScopeLock> GuardWrapper;
+	TOptional<UE::TRWScopeLock<FTransactionallySafeRWLock>> GuardWrapper;
 	UE::AssetRegistry::Private::ELockPriority Priority;
 	FRWScopeLockType LockType;
 };
diff --git a/Engine/Source/Runtime/Core/Public/Misc/TransactionallySafeRWScopeLock.h b/Engine/Source/Runtime/Core/Public/Misc/TransactionallySafeRWScopeLock.h
--- a/Engine/Source/Runtime/Core/Public/Misc/TransactionallySafeRWScopeLock.h
+++ b/Engine/Source/Runtime/Core/Public/Misc/TransactionallySafeRWScopeLock.h
@@ -46,38 +46,3 @@ private:
 
 	FTransactionallySafeRWLock& Lock;
 };
-
-// A transactionally safe read-write scope lock - that uses a transactionally safe read-write lock to back it.
-struct FTransactionallySafeRWScopeLock final
-{
-	UE_NODISCARD_CTOR explicit FTransactionallySafeRWScopeLock(FTransactionallySafeRWLock& InLockObject, FRWScopeLockType InLockType)
-		: LockObject(InLockObject), LockType(InLockType)
-	{
-		if (LockType != SLT_ReadOnly)
-		{
-			LockObject.WriteLock();
-		}
-		else
-		{
-			LockObject.ReadLock();
-		}
-	}
-
-	~FTransactionallySafeRWScopeLock()
-	{
-		if (LockType == SLT_ReadOnly)
-		{
-			LockObject.ReadUnlock();
-		}
-		else
-		{
-			LockObject.WriteUnlock();
-		}
-	}
-
-private:
-	UE_NONCOPYABLE(FTransactionallySafeRWScopeLock);
-
-	FTransactionallySafeRWLock& LockObject;
-	FRWScopeLockType LockType;
-};
diff --git a/Engine/Source/Runtime/Engine/Private/Components/ActorComponent.cpp b/Engine/Source/Runtime/Engine/Private/Components/ActorComponent.cpp
--- a/Engine/Source/Runtime/Engine/Private/Components/ActorComponent.cpp
+++ b/Engine/Source/Runtime/Engine/Private/Components/ActorComponent.cpp
@@ -26,6 +26,8 @@
 #include "HAL/LowLevelMemStats.h"
 #include "Logging/MessageLog.h"
 #include "Misc/MapErrors.h"
+#include "Misc/ScopeRWLock.h"
+#include "Misc/TransactionallySafeRWLock.h"
 #include "Misc/UObjectToken.h"
 #include "Net/Core/PushModel/PushModel.h"
 #include "Net/UnrealNetwork.h"
@@ -647,7 +649,7 @@ void UActorComponent::PostLoad()
 		{
 			if (UCSModifiedProperties_DEPRECATED.Num())
 			{
-				FTransactionallySafeRWScopeLock Lock(AllUCSModifiedPropertiesLock, SLT_Write);
+				UE::TRWScopeLock Lock(AllUCSModifiedPropertiesLock, SLT_Write);
 				AllUCSModifiedProperties.Add(this, MoveTemp(UCSModifiedProperties_DEPRECATED));
 			}
 		}
@@ -2877,7 +2879,7 @@ void UActorComponent::DetermineUCSModifiedProperties()
 			}
 		}
 
-		FTransactionallySafeRWScopeLock Lock(AllUCSModifiedPropertiesLock, SLT_Write);
+		UE::TRWScopeLock Lock(AllUCSModifiedPropertiesLock, SLT_Write);
 		if (UCSModifiedProperties.Num() > 0)
 		{
 			AllUCSModifiedProperties.Add(this, MoveTemp(UCSModifiedProperties));
@@ -2891,7 +2893,7 @@ void UActorComponent::DetermineUCSModifiedProperties()
 
 void UActorComponent::GetUCSModifiedProperties(TSet<const FProperty*>& ModifiedProperties) const
 {
-	FTransactionallySafeRWScopeLock Lock(AllUCSModifiedPropertiesLock, SLT_ReadOnly);
+	UE::TRWScopeLock Lock(AllUCSModifiedPropertiesLock, SLT_ReadOnly);
 	if (TArray<FSimpleMemberReference>* UCSModifiedProperties = AllUCSModifiedProperties.Find(this))
 	{
 		for (const FSimpleMemberReference& MemberReference : *UCSModifiedProperties)
@@ -2903,7 +2905,7 @@ void UActorComponent::GetUCSModifiedProperties(TSet<const FProperty*>& ModifiedP
 
 void UActorComponent::RemoveUCSModifiedProperties(const TArray<FProperty*>& Properties)
 {
-	FTransactionallySafeRWScopeLock Lock(AllUCSModifiedPropertiesLock, SLT_Write);
+	UE::TRWScopeLock Lock(AllUCSModifiedPropertiesLock, SLT_Write);
 	if (TArray<FSimpleMemberReference>* UCSModifiedProperties = AllUCSModifiedProperties.Find(this))
 	{
 		for (FProperty* Property : Properties)
@@ -2917,7 +2919,7 @@ void UActorComponent::RemoveUCSModifiedProperties(const TArray<FProperty*>& Prop
 
 void UActorComponent::ClearUCSModifiedProperties()
 {
-	FTransactionallySafeRWScopeLock Lock(AllUCSModifiedPropertiesLock, SLT_Write);
+	UE::TRWScopeLock Lock(AllUCSModifiedPropertiesLock, SLT_Write);
 	AllUCSModifiedProperties.Remove(this);
 }
 
@@ -2925,7 +2927,7 @@ void UActorComponent::AddReferencedObjects(UObject* InThis, FReferenceCollector&
 {
 	Super::AddReferencedObjects(InThis, Collector);
 
-	FTransactionallySafeRWScopeLock Lock(AllUCSModifiedPropertiesLock, SLT_ReadOnly);
+	UE::TRWScopeLock Lock(AllUCSModifiedPropertiesLock, SLT_ReadOnly);
 	if (TArray<FSimpleMemberReference>* UCSModifiedProperties = AllUCSModifiedProperties.Find(CastChecked<UActorComponent>(InThis)))
 	{
 		for (FSimpleMemberReference& MemberReference : *UCSModifiedProperties)
@@ -2988,7 +2990,7 @@ void UActorComponent::Serialize(FArchive& Ar)
 			TArray<FSimpleMemberReference> UCSModifiedProperties;
 			Ar << UCSModifiedProperties;
 
-			FTransactionallySafeRWScopeLock Lock(AllUCSModifiedPropertiesLock, SLT_Write);
+			UE::TRWScopeLock Lock(AllUCSModifiedPropertiesLock, SLT_Write);
 			if (UCSModifiedProperties.Num() > 0)
 			{
 				AllUCSModifiedProperties.Add(this, MoveTemp(UCSModifiedProperties));
@@ -3000,7 +3002,7 @@ void UActorComponent::Serialize(FArchive& Ar)
 		}
 		else
 		{
-			FTransactionallySafeRWScopeLock Lock(AllUCSModifiedPropertiesLock, SLT_ReadOnly);
+			UE::TRWScopeLock Lock(AllUCSModifiedPropertiesLock, SLT_ReadOnly);
 			if (TArray<FSimpleMemberReference>* UCSModifiedProperties = AllUCSModifiedProperties.Find(this))
 			{
 				Ar << *UCSModifiedProperties;
diff --git a/Engine/Source/Runtime/Engine/Private/ShowFlags.cpp b/Engine/Source/Runtime/Engine/Private/ShowFlags.cpp
--- a/Engine/Source/Runtime/Engine/Private/ShowFlags.cpp
+++ b/Engine/Source/Runtime/Engine/Private/ShowFlags.cpp
@@ -8,7 +8,8 @@
 #include "RenderUtils.h"
 #include "Engine/EngineBaseTypes.h"
 #include "HAL/IConsoleManager.h"
-#include "Misc/TransactionallySafeRWScopeLock.h"
+#include "Misc/ScopeRWLock.h"
+#include "Misc/TransactionallySafeRWLock.h"
 #include "SystemSettings.h"
 
 static bool IsValidNameChar(TCHAR c)
@@ -974,7 +975,7 @@ FEngineShowFlags::ECustomShowFlag FEngineShowFlags::RegisterCustomShowFlag(const
 			}
 		}
 
-		FTransactionallySafeRWScopeLock Lock(CustomShowFlagsInternal::GetLock(), SLT_Write);
+		UE::TRWScopeLock Lock(CustomShowFlagsInternal::GetLock(), SLT_Write);
 		ReturnIndex = (FEngineShowFlags::ECustomShowFlag)CustomShowFlagsInternal::GetRegisteredCustomShowFlags().AddDefaulted(1);
 		CustomShowFlagsInternal::GetDefaultState().Add(1);
 		CustomShowFlagsInternal::GetNameToIndex().Add(Name, ReturnIndex);
@@ -1021,7 +1022,7 @@ FSimpleMulticastDelegate FEngineShowFlags::OnCustomShowFlagRegistered;
 // Private custom show flags functions
 FEngineShowFlags::ECustomShowFlag FEngineShowFlags::FindCustomShowFlagByName(const FString& Name)
 {
-	FTransactionallySafeRWScopeLock Lock(CustomShowFlagsInternal::GetLock(), SLT_ReadOnly);
+	UE::TRWScopeLock Lock(CustomShowFlagsInternal::GetLock(), SLT_ReadOnly);
 
 	if (FEngineShowFlags::ECustomShowFlag* Existing = CustomShowFlagsInternal::GetNameToIndex().Find(Name))
 	{
@@ -1032,13 +1033,13 @@ FEngineShowFlags::ECustomShowFlag FEngineShowFlags::FindCustomShowFlagByName(con
 
 bool FEngineShowFlags::IsRegisteredCustomShowFlag(ECustomShowFlag Index)
 {
-	FTransactionallySafeRWScopeLock Lock(CustomShowFlagsInternal::GetLock(), SLT_ReadOnly);
+	UE::TRWScopeLock Lock(CustomShowFlagsInternal::GetLock(), SLT_ReadOnly);
 	return CustomShowFlagsInternal::GetRegisteredCustomShowFlags().IsValidIndex((uint32)Index);
 }
 
 FString FEngineShowFlags::GetCustomShowFlagName(ECustomShowFlag Index)
 {
-	FTransactionallySafeRWScopeLock Lock(CustomShowFlagsInternal::GetLock(), SLT_ReadOnly);
+	UE::TRWScopeLock Lock(CustomShowFlagsInternal::GetLock(), SLT_ReadOnly);
 	if (CustomShowFlagsInternal::GetRegisteredCustomShowFlags().IsValidIndex((uint32)Index))
 	{
 		return  CustomShowFlagsInternal::GetRegisteredCustomShowFlags()[(uint32)Index].Name;
@@ -1048,7 +1049,7 @@ FString FEngineShowFlags::GetCustomShowFlagName(ECustomShowFlag Index)
 
 FText FEngineShowFlags::GetCustomShowFlagDisplayName(ECustomShowFlag Index)
 {
-	FTransactionallySafeRWScopeLock Lock(CustomShowFlagsInternal::GetLock(), SLT_ReadOnly);
+	UE::TRWScopeLock Lock(CustomShowFlagsInternal::GetLock(), SLT_ReadOnly);
 	if (CustomShowFlagsInternal::GetRegisteredCustomShowFlags().IsValidIndex((uint32)Index))
 	{
 		return CustomShowFlagsInternal::GetRegisteredCustomShowFlags()[(uint32)Index].DisplayName;
@@ -1058,7 +1059,7 @@ FText FEngineShowFlags::GetCustomShowFlagDisplayName(ECustomShowFlag Index)
 
 EShowFlagGroup FEngineShowFlags::GetCustomShowFlagGroup(ECustomShowFlag Index)
 {
-	FTransactionallySafeRWScopeLock Lock(CustomShowFlagsInternal::GetLock(), SLT_ReadOnly);
+	UE::TRWScopeLock Lock(CustomShowFlagsInternal::GetLock(), SLT_ReadOnly);
 	if (CustomShowFlagsInternal::GetRegisteredCustomShowFlags().IsValidIndex((uint32)Index))
 	{
 		return CustomShowFlagsInternal::GetRegisteredCustomShowFlags()[(uint32)Index].Group;
@@ -1068,7 +1069,7 @@ EShowFlagGroup FEngineShowFlags::GetCustomShowFlagGroup(ECustomShowFlag Index)
 
 void FEngineShowFlags::InitCustomShowFlags(EShowFlagInitMode InitMode)
 {
-	FTransactionallySafeRWScopeLock Lock(CustomShowFlagsInternal::GetLock(), SLT_ReadOnly);
+	UE::TRWScopeLock Lock(CustomShowFlagsInternal::GetLock(), SLT_ReadOnly);
 	if (InitMode == ESFIM_All0)
 	{
 		CustomShowFlags.Init(false, CustomShowFlagsInternal::GetDefaultState().Num());
@@ -1081,7 +1082,7 @@ void FEngineShowFlags::InitCustomShowFlags(EShowFlagInitMode InitMode)
 
 void FEngineShowFlags::UpdateNewCustomShowFlags()
 {
-	FTransactionallySafeRWScopeLock Lock(CustomShowFlagsInternal::GetLock(), SLT_ReadOnly);
+	UE::TRWScopeLock Lock(CustomShowFlagsInternal::GetLock(), SLT_ReadOnly);
 	for (int32 i = CustomShowFlags.Num(); i < CustomShowFlagsInternal::GetDefaultState().Num(); ++i)
 	{
 		CustomShowFlags.Add(CustomShowFlagsInternal::GetDefaultState()[i]);
@@ -1090,7 +1091,7 @@ void FEngineShowFlags::UpdateNewCustomShowFlags()
 
 bool FEngineShowFlags::FindCustomShowFlagDisplayName(const FString& Name, FText& OutText)
 {
-	FTransactionallySafeRWScopeLock Lock(CustomShowFlagsInternal::GetLock(), SLT_ReadOnly);
+	UE::TRWScopeLock Lock(CustomShowFlagsInternal::GetLock(), SLT_ReadOnly);
 	if (FEngineShowFlags::ECustomShowFlag* Existing = CustomShowFlagsInternal::GetNameToIndex().Find(Name))
 	{
 		OutText = CustomShowFlagsInternal::GetRegisteredCustomShowFlags()[(int32)*Existing].DisplayName;
@@ -1101,7 +1102,7 @@ bool FEngineShowFlags::FindCustomShowFlagDisplayName(const FString& Name, FText&
 
 void FEngineShowFlags::EngineOverrideCustomShowFlagsFromCVars()
 {
-	FTransactionallySafeRWScopeLock Lock(CustomShowFlagsInternal::GetLock(), SLT_ReadOnly);
+	UE::TRWScopeLock Lock(CustomShowFlagsInternal::GetLock(), SLT_ReadOnly);
 	for (int32 i = 0; i < CustomShowFlagsInternal::GetRegisteredCustomShowFlags().Num(); ++i)
 	{
 		int32 Val = CustomShowFlagsInternal::GetRegisteredCustomShowFlags()[i].CVar->GetInt();
