[AutoRTFM] Fix AutoRTFMTests so it can be compiled in a unity build

We don't currently build these tests with unity builds enabled, but having this fixed helps if we need to test something in a unity build environment.

diff --git a/Engine/Source/Programs/AutoRTFMTests/Private/PackedPointer.cpp b/Engine/Source/Programs/AutoRTFMTests/Private/PackedPointer.cpp
--- a/Engine/Source/Programs/AutoRTFMTests/Private/PackedPointer.cpp
+++ b/Engine/Source/Programs/AutoRTFMTests/Private/PackedPointer.cpp
@@ -4,12 +4,13 @@
 #include <AutoRTFM/AutoRTFM.h>
 #include <inttypes.h>
 
-#pragma pack(1)
+#pragma pack(push,1)
 struct Foo
 {
     uint8_t a;
     uint64_t b;
 };
+#pragma pack(pop)
 
 TEST_CASE("PackedPointer")
 {
diff --git a/Engine/Source/Programs/AutoRTFMTests/Private/UECoreTests.cpp b/Engine/Source/Programs/AutoRTFMTests/Private/UECoreTests.cpp
--- a/Engine/Source/Programs/AutoRTFMTests/Private/UECoreTests.cpp
+++ b/Engine/Source/Programs/AutoRTFMTests/Private/UECoreTests.cpp
@@ -27,7 +27,7 @@
 #include "StructUtils/PropertyBag.h"
 #include "Utils.h"
 
-DEFINE_LOG_CATEGORY_STATIC(LogAutoRTFMTests, Log, All)
+DEFINE_LOG_CATEGORY_STATIC(LogAutoRTFM_UECoreTests, Log, All)
 
 TEST_CASE("UECore.FDelegateHandle")
 {
@@ -2381,14 +2381,14 @@ TEST_CASE("UECore.UE_LOGFMT")
 	{
 		AutoRTFM::Testing::Commit([&]
 		{
-			UE_LOGFMT(LogAutoRTFMTests, Log, "{Animal} says {Sound}", TEXT("Cat"), TEXT("meow!"));
+			UE_LOGFMT(LogAutoRTFM_UECoreTests, Log, "{Animal} says {Sound}", TEXT("Cat"), TEXT("meow!"));
 		});
 	}
 	SECTION("Abort")
 	{
 		AutoRTFM::Testing::Abort([&]
 		{
-			UE_LOGFMT(LogAutoRTFMTests, Log, "{Animal} says {Sound}", TEXT("Cat"), TEXT("meow!"));
+			UE_LOGFMT(LogAutoRTFM_UECoreTests, Log, "{Animal} says {Sound}", TEXT("Cat"), TEXT("meow!"));
 			AutoRTFM::AbortTransaction();
 		});
 	}
