Get `UE_TRACE_LOG_SCOPED_T` working under transactions.

diff --git a/Engine/Source/Programs/AutoRTFMTests/Private/InsightsTests.cpp b/Engine/Source/Programs/AutoRTFMTests/Private/InsightsTests.cpp
--- a/Engine/Source/Programs/AutoRTFMTests/Private/InsightsTests.cpp
+++ b/Engine/Source/Programs/AutoRTFMTests/Private/InsightsTests.cpp
@@ -4,17 +4,42 @@
 #include "Catch2Includes.h"
 #include "CoreTypes.h"
 #include "ProfilingDebugging/CpuProfilerTrace.h"
+#include "Trace/Trace.h"
+#include "Trace/Trace.inl"
 
 #include <AutoRTFM/AutoRTFMTesting.h>
 
 UE_TRACE_CHANNEL_DEFINE(AutoRTFMInsightsChannel)
 
+struct AutoRTFMInsightsDeferer final
+{
+	explicit AutoRTFMInsightsDeferer(UE::Trace::FChannel& InChannel) : Channel(InChannel), bWasEnabled(InChannel)
+	{
+		Channel.Toggle(true);
+	}
+
+	~AutoRTFMInsightsDeferer()
+	{
+		Channel.Toggle(bWasEnabled);
+	}
+
+private:
+	UE::Trace::FChannel& Channel;
+	const bool bWasEnabled;
+};
+
+UE_TRACE_EVENT_BEGIN(Cpu, SomeTraceEvent, NoSync)
+	UE_TRACE_EVENT_FIELD(int32, Foo)
+UE_TRACE_EVENT_END()
+
 TEST_CASE("Insights")
 {
-	CpuChannel.Toggle(true);
+#if CPUPROFILERTRACE_ENABLED
+	AutoRTFMInsightsDeferer Cpu(CpuChannel);
 	REQUIRE(UE_TRACE_CHANNELEXPR_IS_ENABLED(CpuChannel));
-
-	AutoRTFMInsightsChannel.Toggle(true);
+#endif
+	
+	AutoRTFMInsightsDeferer AutoRTFMInsights(AutoRTFMInsightsChannel);
 	REQUIRE(UE_TRACE_CHANNELEXPR_IS_ENABLED(AutoRTFMInsightsChannel));
 
 	SECTION("TRACE_CPUPROFILER_EVENT_DECLARE")
@@ -452,4 +477,14 @@ TEST_CASE("Insights")
 			AutoRTFM::AbortTransaction();
 		});
 	}
+	
+#if CPUPROFILERTRACE_ENABLED
+	SECTION("UE_TRACE_LOG_SCOPED_T")
+	{
+		AutoRTFM::Testing::Commit([&]
+		{
+			UE_TRACE_LOG_SCOPED_T(Cpu, SomeTraceEvent, CpuChannel) << SomeTraceEvent.Foo(42);
+		});
+	}
+#endif
 }
diff --git a/Engine/Source/Runtime/TraceLog/TraceLog.Build.cs b/Engine/Source/Runtime/TraceLog/TraceLog.Build.cs
--- a/Engine/Source/Runtime/TraceLog/TraceLog.Build.cs
+++ b/Engine/Source/Runtime/TraceLog/TraceLog.Build.cs
@@ -14,7 +14,5 @@ public class TraceLog : ModuleRules
 		PublicIncludePathModuleNames.Add("Core");
 
 		PrivateDefinitions.Add("SUPPRESS_PER_MODULE_INLINE_FILE"); // This module does not use core's standard operator new/delete overloads
-		
-		bDisableAutoRTFMInstrumentation = true;
 	}
 }
