[AutoRTFM] Fix closed-transaction reconstruction of named UObjects

The memory validator flagged that performing the Memzero and UObjectBase in-place reconstruction was clobbering tracked memory.
It was right to flag this - aborting the transaction after reconstruction would leave the object in a zero'd state, not the state before reconstruction.

To fix this record the writes in the open block only if the object is being reconstructed:
* If a this is a new object, we *do not* want to record the write as reverting the memory would bring the object to an uninitialized memory state, which will cause the GC to access uninitialized memory.
* If this is a recycled object, then we *do* want to record the write so an abort will return the object to the previous object's state.

We disable validation on this write because this spans the entire object's memory including fields that may be accessed only in the open (e.g. object flags which are set atomically), and mixing closed writes with open writes will trigger the memory validator.

Also:
* Switch FUObjectAnnotationDense and FPropertyBagRepository to using the transactionally-safe locks. This is required as the new unit tests exercise new code paths with the WITH_EDITORONLY_DATA define enabled.
* Fix (and add tests for) the `bNoMemoryValidation` flag not getting applied to the chunks of memory past the initial chunk for a large write.

diff --git a/Engine/Source/Programs/AutoRTFMTests/Private/Open.cpp b/Engine/Source/Programs/AutoRTFMTests/Private/Open.cpp
--- a/Engine/Source/Programs/AutoRTFMTests/Private/Open.cpp
+++ b/Engine/Source/Programs/AutoRTFMTests/Private/Open.cpp
@@ -23,10 +23,10 @@ TEST_CASE("Open")
     bool bDidRun = false;
     REQUIRE(
         AutoRTFM::ETransactionResult::Committed ==
-        AutoRTFM::Transact([&] ()
+        AutoRTFM::Transact([&]
         {
-            AutoRTFM::Open([&] () { bDidRun = true; });
-            AutoRTFM::Open([&] () { REQUIRE(bDidRun); });
+            AutoRTFM::Open([&] { bDidRun = true; });
+            AutoRTFM::Open([&] { REQUIRE(bDidRun); });
         }));
     REQUIRE(bDidRun);
 }
@@ -46,7 +46,7 @@ TEST_CASE("Open.Large")
     m[6].push_back(9);
     REQUIRE(
         AutoRTFM::ETransactionResult::Committed ==
-        AutoRTFM::Transact([&] () {
+        AutoRTFM::Transact([&] {
             x = 5;
 
             for (size_t n = 10; n--;)
@@ -59,7 +59,7 @@ TEST_CASE("Open.Large")
             m[12].push_back(13);
             m[12].push_back(14);
 
-            AutoRTFM::Open([&] () {
+            AutoRTFM::Open([&] {
 #if 0
                 // These checks are UB, because the open is interacting with transactional data!
                 REQUIRE(x == 42);
@@ -107,10 +107,10 @@ TEST_CASE("Open.Atomics")
     std::atomic<bool> bDidRun = false;
     REQUIRE(
         AutoRTFM::ETransactionResult::Committed ==
-        AutoRTFM::Transact([&] ()
+        AutoRTFM::Transact([&]
         {
-            AutoRTFM::Open([&] () { bDidRun = true; });
-            AutoRTFM::Open([&] () { REQUIRE(bDidRun); });
+            AutoRTFM::Open([&] { bDidRun = true; });
+            AutoRTFM::Open([&] { REQUIRE(bDidRun); });
         }));
     REQUIRE(bDidRun);
 }
@@ -128,7 +128,7 @@ TEST_CASE("Open.ReturnValue")
     SECTION("int")
     {
         int Value = 10;
-        AutoRTFM::ETransactionResult Result = AutoRTFM::Transact([&] ()
+        AutoRTFM::ETransactionResult Result = AutoRTFM::Transact([&]
             {
                 int Got = AutoRTFM::Open([] { return 42; });
                 AutoRTFM::Open([&] { Value = Got; });
@@ -140,7 +140,7 @@ TEST_CASE("Open.ReturnValue")
     SECTION("char*")
     {
         std::string Value = "<unassigned>";
-        AutoRTFM::ETransactionResult Result = AutoRTFM::Transact([&] ()
+        AutoRTFM::ETransactionResult Result = AutoRTFM::Transact([&]
             {
                 // Note: AutoRTFM::Open() is returning a const char*
                 std::string Got = AutoRTFM::Open([] { return "meow"; });
@@ -154,7 +154,7 @@ TEST_CASE("Open.ReturnValue")
     {
         int Int = 0;
         std::string String = "<unassigned>";
-        AutoRTFM::ETransactionResult Result = AutoRTFM::Transact([&] ()
+        AutoRTFM::ETransactionResult Result = AutoRTFM::Transact([&]
             {
                 auto [I, S] = AutoRTFM::Open([] 
                     { 
@@ -241,6 +241,11 @@ TEST_CASE("Open.Collision")
 {
     AUTORTFM_SCOPED_ENABLE_MEMORY_VALIDATION_AS_WARNING();
 
+    struct FLargeStruct
+    {
+        int V[256];
+    };
+
     AutoRTFMTestUtils::FCaptureWarningContext WarningContext;
     int I = 0;
     int J = 0;
@@ -249,10 +254,10 @@ TEST_CASE("Open.Collision")
     {
         SECTION("Different Memory Locations")
         {
-            AutoRTFM::Transact([&] ()
+            AutoRTFM::Transact([&]
             {
                 I = 42;
-                AutoRTFM::Open([&] () { J = 24; });
+                AutoRTFM::Open([&] { J = 24; });
             });
 
             REQUIRE(I == 42);
@@ -262,13 +267,13 @@ TEST_CASE("Open.Collision")
 
         SECTION("Transact(Open(RecordOpenWriteNoMemoryValidation()), Open(Assign))")
         {
-            AutoRTFM::Transact([&] ()
+            AutoRTFM::Transact([&]
             {
-                AutoRTFM::Open([&] () 
+                AutoRTFM::Open([&]
                 {
                     AutoRTFM::RecordOpenWriteNoMemoryValidation(&I);
                 });
-                AutoRTFM::Open([&] () 
+                AutoRTFM::Open([&]
                 {
                     I = 42;
                 });
@@ -278,11 +283,36 @@ TEST_CASE("Open.Collision")
             REQUIRE(WarningContext.GetWarnings().empty());
         }
 
+        SECTION("Transact(Open(RecordOpenWriteNoMemoryValidation()), Open(Assign)) <large>")
+        {
+            FLargeStruct LargeStruct;
+            AutoRTFM::Transact([&]
+            {
+                AutoRTFM::Open([&]
+                {
+                    AutoRTFM::RecordOpenWriteNoMemoryValidation(&LargeStruct);
+                });
+                AutoRTFM::Open([&]
+                {
+                    for (int& V : LargeStruct.V)
+                    {
+                        V = 42;
+                    }
+                });
+            });
+
+            for (int V : LargeStruct.V)
+            {
+                REQUIRE(V == 42);
+            }
+            REQUIRE(WarningContext.GetWarnings().empty());
+        }
+
         SECTION("Transact(Open(Assign), Assign)")
         {
-            AutoRTFM::Transact([&] ()
+            AutoRTFM::Transact([&]
             {
-                AutoRTFM::Open([&] () { I = 24; });
+                AutoRTFM::Open([&] { I = 24; });
                 I = 42;
             });
 
@@ -292,12 +322,12 @@ TEST_CASE("Open.Collision")
 
         SECTION("Transact(Assign, Open(Transact(Assign)))")
         {
-            AutoRTFM::Transact([&] ()
+            AutoRTFM::Transact([&]
             {
                 I = 24;
-                AutoRTFM::Open([&] ()
+                AutoRTFM::Open([&]
                 {
-                    AutoRTFM::Transact([&] ()
+                    AutoRTFM::Transact([&]
                     {
                         I = 42;
                     });
@@ -310,11 +340,11 @@ TEST_CASE("Open.Collision")
 
         SECTION("Transact(Open(Transact(Assign)), Assign)")
         {
-            AutoRTFM::Transact([&] ()
+            AutoRTFM::Transact([&]
             {
-                AutoRTFM::Open([&] ()
+                AutoRTFM::Open([&]
                 {
-                    AutoRTFM::Transact([&] ()
+                    AutoRTFM::Transact([&]
                     {
                         I = 24;
                     });
@@ -329,10 +359,10 @@ TEST_CASE("Open.Collision")
 
     SECTION("Transact(Assign, Open(Assign))")
     {
-        AutoRTFM::Transact([&] ()
+        AutoRTFM::Transact([&]
         {
             I = 42;
-            AutoRTFM::Open([&] () { I = 24; });
+            AutoRTFM::Open([&] { I = 24; });
         });
 
         REQUIRE_THAT(WarningContext.GetWarnings(),
@@ -341,12 +371,12 @@ TEST_CASE("Open.Collision")
 
     SECTION("Transact(Assign, Transact(Open(Assign)))")
     {
-        AutoRTFM::Transact([&] ()
+        AutoRTFM::Transact([&]
         {
             I = 42;
-            AutoRTFM::Transact([&] ()
+            AutoRTFM::Transact([&]
             {
-                AutoRTFM::Open([&] () { I = 24; });
+                AutoRTFM::Open([&] { I = 24; });
             });
         });
 
@@ -356,14 +386,14 @@ TEST_CASE("Open.Collision")
 
     SECTION("Transact(Assign, Transact(Transact(Open(Assign))))")
     {
-        AutoRTFM::Transact([&] ()
+        AutoRTFM::Transact([&]
         {
             I = 42;
-            AutoRTFM::Transact([&] ()
+            AutoRTFM::Transact([&]
             {
-                AutoRTFM::Transact([&] ()
+                AutoRTFM::Transact([&]
                 {
-                    AutoRTFM::Open([&] ()
+                    AutoRTFM::Open([&]
                     {
                         I = 24;
                     });
@@ -377,14 +407,14 @@ TEST_CASE("Open.Collision")
 
     SECTION("Transact(Transact(Assign, Transact(Open(Assign))))")
     {
-        AutoRTFM::Transact([&] ()
+        AutoRTFM::Transact([&]
         {
-            AutoRTFM::Transact([&] ()
+            AutoRTFM::Transact([&]
             {
                 I = 42;
-                AutoRTFM::Transact([&] ()
+                AutoRTFM::Transact([&]
                 {
-                    AutoRTFM::Open([&] ()
+                    AutoRTFM::Open([&]
                     {
                         I = 24;
                     });
@@ -398,13 +428,13 @@ TEST_CASE("Open.Collision")
 
     SECTION("Transact(Transact(Assign, Open(Assign, Transact())))")
     {
-        AutoRTFM::Transact([&] ()
+        AutoRTFM::Transact([&]
         {
             I = 42;
-            AutoRTFM::Open([&] ()
+            AutoRTFM::Open([&]
             {
                 I = 24;
-                AutoRTFM::Transact([&] () {});
+                AutoRTFM::Transact([&] {});
             });
         });
 
@@ -414,12 +444,12 @@ TEST_CASE("Open.Collision")
 
     SECTION("Transact(Transact(Assign, Open(Transact(), Assign)))")
     {
-        AutoRTFM::Transact([&] ()
+        AutoRTFM::Transact([&]
         {
             I = 42;
-            AutoRTFM::Open([&] ()
+            AutoRTFM::Open([&]
             {
-                AutoRTFM::Transact([&] () {});
+                AutoRTFM::Transact([&] {});
                 I = 24;
             });
         });
@@ -430,7 +460,7 @@ TEST_CASE("Open.Collision")
 
     SECTION("Transact(Assign, CallOpen(Assign))")
     {
-        AutoRTFM::Transact([&] ()
+        AutoRTFM::Transact([&]
         {
             I = 42;
             AssignIntPointer(&I, 24);
@@ -442,10 +472,10 @@ TEST_CASE("Open.Collision")
 
     SECTION("Transact(Assign, Transact(CallOpen(Assign)))")
     {
-        AutoRTFM::Transact([&] ()
+        AutoRTFM::Transact([&]
         {
             I = 42;
-            AutoRTFM::Transact([&] ()
+            AutoRTFM::Transact([&]
             {
                 AssignIntPointer(&I, 24);
             });
@@ -457,7 +487,7 @@ TEST_CASE("Open.Collision")
 
     SECTION("Transact(Assign, OpenNoValidation(Assign))")
     {
-        AutoRTFM::Transact([&] ()
+        AutoRTFM::Transact([&]
         {
             I = 42;
             AutoRTFM::Open<AutoRTFM::EMemoryValidationLevel::Disabled>([&]
@@ -480,7 +510,7 @@ TEST_CASE("Open.Collision")
             }
         };
 
-        AutoRTFM::Transact([&] ()
+        AutoRTFM::Transact([&]
         {
             I = 42;
             S::AssignToInt(I);
@@ -491,7 +521,7 @@ TEST_CASE("Open.Collision")
 
     SECTION("Transact(Assign, Open(Assign, OpenNoValidation()))")
     {
-        AutoRTFM::Transact([&] ()
+        AutoRTFM::Transact([&]
         {
             I = 42;
             AutoRTFM::Open([&]
@@ -507,7 +537,7 @@ TEST_CASE("Open.Collision")
     
     SECTION("Transact(Assign, Open(Assign, Close(OpenNoValidation())))")
     {
-        AutoRTFM::Transact([&] ()
+        AutoRTFM::Transact([&]
         {
             I = 42;
             AutoRTFM::Open([&]
@@ -526,13 +556,13 @@ TEST_CASE("Open.Collision")
 
     SECTION("Transact(Open(RecordOpenWrite()), Open(Assign))")
     {
-        AutoRTFM::Transact([&] ()
+        AutoRTFM::Transact([&]
         {
-            AutoRTFM::Open([&] () 
+            AutoRTFM::Open([&]
             {
                 AutoRTFM::RecordOpenWrite(&I);
             });
-            AutoRTFM::Open([&] () 
+            AutoRTFM::Open([&]
             {
                 I = 42;
             });
@@ -545,17 +575,17 @@ TEST_CASE("Open.Collision")
 
     SECTION("Transact(Open(RecordOpenWriteNoMemoryValidation()), Open(RecordOpenWrite()), Open(Assign))")
     {
-        AutoRTFM::Transact([&] ()
+        AutoRTFM::Transact([&]
         {
-            AutoRTFM::Open([&] () 
+            AutoRTFM::Open([&]
             {
                 AutoRTFM::RecordOpenWriteNoMemoryValidation(&I);
             });
-            AutoRTFM::Open([&] () 
+            AutoRTFM::Open([&]
             {
                 AutoRTFM::RecordOpenWrite(&I);
             });
-            AutoRTFM::Open([&] () 
+            AutoRTFM::Open([&]
             {
                 I = 42;
             });
@@ -565,4 +595,27 @@ TEST_CASE("Open.Collision")
             Catch::Matchers::VectorContains(FString(AutoRTFMTestUtils::kMemoryModifiedWarning)));
     }
 
+
+    SECTION("Transact(Open(RecordOpenWriteNoMemoryValidation()), Open(RecordOpenWrite()), Open(Assign)) <large>")
+    {
+        FLargeStruct LargeStruct;
+        AutoRTFM::Transact([&]
+        {
+            AutoRTFM::Open([&]
+            {
+                AutoRTFM::RecordOpenWriteNoMemoryValidation(&LargeStruct);
+            });
+            AutoRTFM::Open([&]
+            {
+                AutoRTFM::RecordOpenWrite(&LargeStruct);
+            });
+            AutoRTFM::Open([&]
+            {
+                LargeStruct.V[255] = 42;
+            });
+        });
+
+        REQUIRE_THAT(WarningContext.GetWarnings(),
+            Catch::Matchers::VectorContains(FString(AutoRTFMTestUtils::kMemoryModifiedWarning)));
+    }
 }
diff --git a/Engine/Source/Programs/AutoRTFMTests/Private/UObjectTests.cpp b/Engine/Source/Programs/AutoRTFMTests/Private/UObjectTests.cpp
--- a/Engine/Source/Programs/AutoRTFMTests/Private/UObjectTests.cpp
+++ b/Engine/Source/Programs/AutoRTFMTests/Private/UObjectTests.cpp
@@ -10,8 +10,9 @@ THIRD_PARTY_INCLUDES_START
 THIRD_PARTY_INCLUDES_END
 
 #include "MyAutoRTFMTestObject.h"
-#include <AutoRTFM/AutoRTFM.h>
+#include "AutoRTFM/AutoRTFMTesting.h"
 #include "UObject/GCObject.h"
+#include "UObject/Package.h"
 #include "UObject/ReachabilityAnalysis.h"
 #include "UObject/UObjectAnnotation.h"
 
@@ -21,7 +22,7 @@ TEST_CASE("UObject.NewObject")
 	{
 		UMyAutoRTFMTestObject* Object = nullptr;
 
-		AutoRTFM::Commit([&]
+		AutoRTFM::Testing::Commit([&]
 			{
 				Object = NewObject<UMyAutoRTFMTestObject>();
 			});
@@ -34,11 +35,11 @@ TEST_CASE("UObject.NewObject")
 	{
 		UMyAutoRTFMTestObject* Object = nullptr;
 
-		REQUIRE(AutoRTFM::ETransactionResult::AbortedByRequest == AutoRTFM::Transact([&]
+		AutoRTFM::Testing::Abort([&]
 			{
 				Object = NewObject<UMyAutoRTFMTestObject>();
 				AutoRTFM::AbortTransaction();
-			}));
+			});
 
 		REQUIRE(nullptr == Object);
 	}
@@ -51,7 +52,7 @@ TEST_CASE("UObject.NewObjectWithOuter")
 		UMyAutoRTFMTestObject* Outer = NewObject<UMyAutoRTFMTestObject>();
 		UMyAutoRTFMTestObject* Object = nullptr;
 
-		AutoRTFM::Commit([&]
+		AutoRTFM::Testing::Commit([&]
 			{
 				Object = NewObject<UMyAutoRTFMTestObject>(Outer);
 			});
@@ -67,17 +68,55 @@ TEST_CASE("UObject.NewObjectWithOuter")
 		UMyAutoRTFMTestObject* Outer = NewObject<UMyAutoRTFMTestObject>();
 		UMyAutoRTFMTestObject* Object = nullptr;
 
-		REQUIRE(AutoRTFM::ETransactionResult::AbortedByRequest == AutoRTFM::Transact([&]
+		AutoRTFM::Testing::Abort([&]
 			{
 				Object = NewObject<UMyAutoRTFMTestObject>(Outer);
 				AutoRTFM::AbortTransaction();
-			}));
+			});
 
 		REQUIRE(nullptr == Object);
 		REQUIRE(42 == Outer->Value);
 	}
 }
 
+TEST_CASE("UObject.ReconstructObject")
+{
+	FName Name = TEXT("MyObjectToBeReplaced");
+	UObject* Outer = static_cast<UObject*>(GetTransientPackage());
+	UMyAutoRTFMTestObject* Old = NewObject<UMyAutoRTFMTestObject>(Outer, Name);
+	REQUIRE(42 == Old->Value);
+	Old->Value = 123;
+
+	SECTION("Create")
+	{
+		UMyAutoRTFMTestObject* New = nullptr;
+
+		AutoRTFM::Testing::Commit([&]
+			{
+				// Reconstruct the object by using the same name.
+				New = NewObject<UMyAutoRTFMTestObject>(Outer, Name);
+			});
+
+		REQUIRE(Old == New);
+		REQUIRE(42 == New->Value);
+	}
+
+	SECTION("Abort")
+	{
+		UMyAutoRTFMTestObject* New = nullptr;
+
+		AutoRTFM::Testing::Abort([&]
+			{
+				// Reconstruct the object by using the same name, then immediately abort.
+				New = NewObject<UMyAutoRTFMTestObject>(Outer, Name);
+				AutoRTFM::AbortTransaction();
+			});
+
+		REQUIRE(nullptr == New);
+		REQUIRE(123 == Old->Value);
+	}
+}
+
 // This is a copy of the helper function in TestGarbageCollector.cpp.
 int32 PerformGarbageCollectionWithIncrementalReachabilityAnalysis(TFunctionRef<bool(int32)> ReachabilityIterationCallback)
 {
@@ -136,10 +175,10 @@ TEST_CASE("UObject.MarkAsReachable")
 			return true;
 		}
 
-		REQUIRE(AutoRTFM::ETransactionResult::Committed == AutoRTFM::Transact([&]
+		AutoRTFM::Testing::Commit([&]
 		{
 			Object->MarkAsReachable();
-		}));
+		});
 
 		return false;
 	});
diff --git a/Engine/Source/Runtime/Core/Private/AutoRTFM/Transaction.h b/Engine/Source/Runtime/Core/Private/AutoRTFM/Transaction.h
--- a/Engine/Source/Runtime/Core/Private/AutoRTFM/Transaction.h
+++ b/Engine/Source/Runtime/Core/Private/AutoRTFM/Transaction.h
@@ -119,7 +119,7 @@ public:
 
     // Record that a write is about to occur at the given LogicalAddress of Size bytes.
     void RecordWrite(void* LogicalAddress, size_t Size, bool bNoMemoryValidation = false);
-    void RecordWriteMaxPageSized(void* LogicalAddress, size_t Size, bool bNoMemoryValidation = false);
+    void RecordWriteMaxPageSized(void* LogicalAddress, size_t Size, bool bNoMemoryValidation);
     template<unsigned SIZE> void RecordWrite(void* LogicalAddress);
 	template<unsigned SIZE> void RecordWriteInsertedSlow(void* LogicalAddress);
 	template<unsigned SIZE> void RecordWriteNotInsertedSlow(void* LogicalAddress);
diff --git a/Engine/Source/Runtime/Core/Private/AutoRTFM/TransactionInlines.h b/Engine/Source/Runtime/Core/Private/AutoRTFM/TransactionInlines.h
--- a/Engine/Source/Runtime/Core/Private/AutoRTFM/TransactionInlines.h
+++ b/Engine/Source/Runtime/Core/Private/AutoRTFM/TransactionInlines.h
@@ -98,7 +98,7 @@ UE_AUTORTFM_FORCEINLINE bool FTransaction::ShouldRecordWrite(void* LogicalAddres
 }
 
 AUTORTFM_NO_ASAN UE_AUTORTFM_FORCEINLINE
-void FTransaction::RecordWriteMaxPageSized(void* LogicalAddress, size_t Size, bool bNoMemoryValidation /* = false */)
+void FTransaction::RecordWriteMaxPageSized(void* LogicalAddress, size_t Size, bool bNoMemoryValidation)
 {
     void* CopyAddress = WriteLogBumpAllocator.Allocate(Size);
     memcpy(CopyAddress, LogicalAddress, Size);
@@ -159,7 +159,7 @@ void FTransaction::RecordWrite(void* LogicalAddress, size_t Size, bool bNoMemory
 
     for (; (I + FWriteLogBumpAllocator::MaxSize) < Size; I += FWriteLogBumpAllocator::MaxSize)
     {
-        RecordWriteMaxPageSized(Address + I, FWriteLogBumpAllocator::MaxSize);
+        RecordWriteMaxPageSized(Address + I, FWriteLogBumpAllocator::MaxSize, bNoMemoryValidation);
     }
 
     // Remainder at the end of the memcpy.
diff --git a/Engine/Source/Runtime/CoreUObject/Internal/UObject/PropertyBagRepository.h b/Engine/Source/Runtime/CoreUObject/Internal/UObject/PropertyBagRepository.h
--- a/Engine/Source/Runtime/CoreUObject/Internal/UObject/PropertyBagRepository.h
+++ b/Engine/Source/Runtime/CoreUObject/Internal/UObject/PropertyBagRepository.h
@@ -4,7 +4,7 @@
 
 #include "CoreTypes.h"
 #include "Containers/Map.h"
-#include "HAL/CriticalSection.h"
+#include "Misc/TransactionallySafeMutex.h"
 #include "UObject/GCObject.h"
 #include "Templates/FunctionFwd.h"
 
@@ -33,7 +33,7 @@ class FPropertyBagRepository : public FGCObject
 
 private:
 	friend class FPropertyBagRepositoryLock;
-	mutable FCriticalSection CriticalSection;
+	mutable FTransactionallySafeMutex CriticalSection;
 
 	// Lifetimes/ownership:
 	// Managed within UObjectBase and synced with object lifetime. The repo tracks pointers to bags, not the bags themselves.
diff --git a/Engine/Source/Runtime/CoreUObject/Private/UObject/UObjectGlobals.cpp b/Engine/Source/Runtime/CoreUObject/Private/UObject/UObjectGlobals.cpp
--- a/Engine/Source/Runtime/CoreUObject/Private/UObject/UObjectGlobals.cpp
+++ b/Engine/Source/Runtime/CoreUObject/Private/UObject/UObjectGlobals.cpp
@@ -3499,7 +3499,8 @@ UObject* StaticAllocateObject
 	UObject* Obj = NULL;
 	if(InName == NAME_None)
 	{
-		AutoRTFM::Open([&]{
+		UE_AUTORTFM_OPEN
+		{
 #if WITH_EDITOR
 			if ( GOutputCookingWarnings && GetTransientPackage() != InOuter->GetOutermost() )
 			{
@@ -3510,7 +3511,7 @@ UObject* StaticAllocateObject
 			{
 				InName = MakeUniqueObjectName(InOuter, InClass);
 			}
-		});
+		};
 	}
 	else
 	{
@@ -3701,11 +3702,21 @@ UObject* StaticAllocateObject
 			*InClass->GetPathName(), *GetPathNameSafe(InOuter), *InName.ToString(), *OldRemoteId.ToString(), *RemoteId.ToString());
 #endif // UE_WITH_REMOTE_OBJECT_HANDLE
 
-		// perform the UObjectBase construction in the open - we expect the GC to invoke
-		// the destructor in the case of transaction abort
-		// Validation temporarily disabled until SOL-7253 is investigated
 		UE_AUTORTFM_OPEN_NO_VALIDATION
 		{
+			// If a this is a new object, we *do not* want to record the write as reverting the
+			// memory would bring the object to an uninitialized memory state, which will
+			// cause the GC to access uninitialized memory.
+			// If this is a recycled object, then we *do* want to record the write so an abort
+			// will return the object to the previous object's state.
+			if (bWasConstructedOnOldObject)
+			{
+				// We disable validation on this write because this spans the entire object's memory
+				// including fields that may be accessed only in the open (e.g. object flags which
+				// are set atomically), and mixing closed writes with open writes will trigger the
+				// memory validator.
+				AutoRTFM::RecordOpenWriteNoMemoryValidation((void *)Obj, TotalSize);
+			}
 			FMemory::Memzero((void *)Obj, TotalSize);
 			new ((void *)Obj) UObjectBase(const_cast<UClass*>(InClass), InFlags|RF_NeedInitialization, InternalSetFlags, InOuter, InName, OldIndex, OldSerialNumber == 0 ? SerialNumber : OldSerialNumber, 
 #if UE_WITH_REMOTE_OBJECT_HANDLE
@@ -4265,7 +4276,7 @@ void FObjectInitializer::PostConstructInit()
 	// need a fence to guarantee that the cleared flag is only visible to other threads
 	// after the other initialization-related writes
 	std::atomic_thread_fence(std::memory_order_release);
-	UE_AUTORTFM_OPEN{ Obj->ClearFlags(RF_NeedInitialization); };
+	Obj->ClearFlags(RF_NeedInitialization);
 
 	// clear the object pointer so we can guard against running this function again
 	Obj = nullptr;
diff --git a/Engine/Source/Runtime/CoreUObject/Public/UObject/UObjectAnnotation.h b/Engine/Source/Runtime/CoreUObject/Public/UObject/UObjectAnnotation.h
--- a/Engine/Source/Runtime/CoreUObject/Public/UObject/UObjectAnnotation.h
+++ b/Engine/Source/Runtime/CoreUObject/Public/UObject/UObjectAnnotation.h
@@ -867,7 +867,7 @@ public:
 
 		UE_AUTORTFM_OPEN
 		{
-			FTransactionallySafeReadScopeLock AnnotationArrayLock(AnnotationArrayCritical);
+			UE::TReadScopeLock AnnotationArrayLock(AnnotationArrayCritical);
 
 			const int32 ChunkIndex = Index / NumAnnotationsPerChunk;
 			if (ChunkIndex < Chunks.Num())
@@ -1093,13 +1093,13 @@ public:
 	 */
 	void AddAnnotation(int32 Index, const TAnnotation& Annotation)
 	{
-		FRWScopeLock AnnotationArrayLock(AnnotationArrayCritical, SLT_Write);
+		UE::TWriteScopeLock AnnotationArrayLock(AnnotationArrayCritical);
 		AddAnnotationInternal(Index, Annotation);
 	}
 
 	void AddAnnotation(int32 Index, TAnnotation&& Annotation)
 	{
-		FRWScopeLock AnnotationArrayLock(AnnotationArrayCritical, SLT_Write);
+		UE::TWriteScopeLock AnnotationArrayLock(AnnotationArrayCritical);
 		AddAnnotationInternal(Index, MoveTemp(Annotation));
 	}
 
@@ -1156,7 +1156,7 @@ public:
 	 */
 	void RemoveAnnotation(int32 Index)
 	{
-		FRWScopeLock AnnotationArrayLock(AnnotationArrayCritical, SLT_Write);
+		UE::TWriteScopeLock AnnotationArrayLock(AnnotationArrayCritical);
 		RemoveAnnotationInternal(Index);
 	}
 
@@ -1177,7 +1177,7 @@ public:
 	 */
 	void RemoveAllAnnotations()
 	{
-		FRWScopeLock AnnotationArrayLock(AnnotationArrayCritical, SLT_Write);
+		UE::TWriteScopeLock AnnotationArrayLock(AnnotationArrayCritical);
 		bool bHadElements = (AnnotationArray.Num() > 0);
 		AnnotationArray.Empty();
 		if (bHadElements)
@@ -1217,7 +1217,7 @@ public:
 	FORCEINLINE TAnnotation GetAnnotation(int32 Index)
 	{
 		check(Index >= 0);
-		FRWScopeLock AnnotationArrayLock(AnnotationArrayCritical, SLT_ReadOnly);
+		UE::TReadScopeLock AnnotationArrayLock(AnnotationArrayCritical);
 		if (Index < AnnotationArray.Num())
 		{
 			return AnnotationArray[Index];
@@ -1246,7 +1246,7 @@ public:
 	 */
 	FORCEINLINE TAnnotation& GetAnnotationRef(int32 Index)
 	{
-		FRWScopeLock AnnotationArrayLock(AnnotationArrayCritical, SLT_Write);
+		UE::TWriteScopeLock AnnotationArrayLock(AnnotationArrayCritical);
 		if (Index >= AnnotationArray.Num())
 		{
 			AddAnnotationInternal(Index, TAnnotation());
@@ -1266,7 +1266,7 @@ private:
 	 * Map from live objects to an annotation
 	 */
 	TArray<TAnnotation> AnnotationArray;
-	FRWLock AnnotationArrayCritical;
+	FTransactionallySafeRWLock AnnotationArrayCritical;
 
 };
 
diff --git a/Engine/Source/Runtime/Engine/Private/Tests/AutoRTFM/AutoRTFMActorComponentTests.cpp b/Engine/Source/Runtime/Engine/Private/Tests/AutoRTFM/AutoRTFMActorComponentTests.cpp
--- a/Engine/Source/Runtime/Engine/Private/Tests/AutoRTFM/AutoRTFMActorComponentTests.cpp
+++ b/Engine/Source/Runtime/Engine/Private/Tests/AutoRTFM/AutoRTFMActorComponentTests.cpp
@@ -3,7 +3,7 @@
 #include "Misc/AutomationTest.h"
 #include "Animation/AnimInstance.h"
 #include "Animation/Skeleton.h"
-#include "AutoRTFM/AutoRTFM.h"
+#include "AutoRTFM/AutoRTFMTesting.h"
 #include "AutoRTFMTestActor.h"
 #include "AutoRTFMTestAnotherActor.h"
 #include "AutoRTFMTestBodySetup.h"
@@ -20,6 +20,7 @@
 #include "Physics/Experimental/PhysScene_Chaos.h"
 #include "PhysicsProxy/SingleParticlePhysicsProxy.h"
 #include "Rendering/SkeletalMeshRenderData.h"
+#include "UObject/Package.h"
 #include "UObject/UObjectArray.h"
 #include "UObject/UObjectGlobals.h"
 
@@ -492,6 +493,27 @@ AUTORTFM_ACTOR_COMPONENT_TEST(WorldBlockTillLevelStreamingCompleted)
 	TestTrueExpr(AutoRTFM::ETransactionResult::Committed == Result);
 }
 
+AUTORTFM_ACTOR_COMPONENT_TEST(ReconstructObject)
+{
+	FName Name = TEXT("MyObjectToBeReplaced");
+	UObject* Outer = static_cast<UObject*>(GetTransientPackage());
+	AActor* Old = NewObject<AActor>(Outer, Name);
+	TestTrueExpr(0 == Old->ActorCategory);
+	Old->ActorCategory = 123;
+
+	AActor* New = nullptr;
+
+	AutoRTFM::Testing::Abort([&]
+		{
+			// Reconstruct the object by using the same name, then immediately abort.
+			New = NewObject<AActor>(Outer, Name);
+			AutoRTFM::AbortTransaction();
+		});
+
+	TestTrueExpr(nullptr == New);
+	TestTrueExpr(123 == Old->ActorCategory);
+}
+
 }  // anonymous namespace
 
 #undef AUTORTFM_ACTOR_COMPONENT_TEST
