[AutoRTFM] Couple of fixes for bNoMemoryValidation write flag

Ensure this flag is applied to the chunks of memory past the initial chunk for a large write.
Ensure this flag is applied to outer transactions when committing a nested transaction.

diff --git a/Engine/Source/Programs/AutoRTFMTests/Private/Open.cpp b/Engine/Source/Programs/AutoRTFMTests/Private/Open.cpp
--- a/Engine/Source/Programs/AutoRTFMTests/Private/Open.cpp
+++ b/Engine/Source/Programs/AutoRTFMTests/Private/Open.cpp
@@ -23,10 +23,10 @@ TEST_CASE("Open")
     bool bDidRun = false;
     REQUIRE(
         AutoRTFM::ETransactionResult::Committed ==
-        AutoRTFM::Transact([&] ()
+        AutoRTFM::Transact([&]
         {
-            AutoRTFM::Open([&] () { bDidRun = true; });
-            AutoRTFM::Open([&] () { REQUIRE(bDidRun); });
+            AutoRTFM::Open([&] { bDidRun = true; });
+            AutoRTFM::Open([&] { REQUIRE(bDidRun); });
         }));
     REQUIRE(bDidRun);
 }
@@ -46,7 +46,7 @@ TEST_CASE("Open.Large")
     m[6].push_back(9);
     REQUIRE(
         AutoRTFM::ETransactionResult::Committed ==
-        AutoRTFM::Transact([&] () {
+        AutoRTFM::Transact([&] {
             x = 5;
 
             for (size_t n = 10; n--;)
@@ -59,7 +59,7 @@ TEST_CASE("Open.Large")
             m[12].push_back(13);
             m[12].push_back(14);
 
-            AutoRTFM::Open([&] () {
+            AutoRTFM::Open([&] {
 #if 0
                 // These checks are UB, because the open is interacting with transactional data!
                 REQUIRE(x == 42);
@@ -107,10 +107,10 @@ TEST_CASE("Open.Atomics")
     std::atomic<bool> bDidRun = false;
     REQUIRE(
         AutoRTFM::ETransactionResult::Committed ==
-        AutoRTFM::Transact([&] ()
+        AutoRTFM::Transact([&]
         {
-            AutoRTFM::Open([&] () { bDidRun = true; });
-            AutoRTFM::Open([&] () { REQUIRE(bDidRun); });
+            AutoRTFM::Open([&] { bDidRun = true; });
+            AutoRTFM::Open([&] { REQUIRE(bDidRun); });
         }));
     REQUIRE(bDidRun);
 }
@@ -128,7 +128,7 @@ TEST_CASE("Open.ReturnValue")
     SECTION("int")
     {
         int Value = 10;
-        AutoRTFM::ETransactionResult Result = AutoRTFM::Transact([&] ()
+        AutoRTFM::ETransactionResult Result = AutoRTFM::Transact([&]
             {
                 int Got = AutoRTFM::Open([] { return 42; });
                 AutoRTFM::Open([&] { Value = Got; });
@@ -140,7 +140,7 @@ TEST_CASE("Open.ReturnValue")
     SECTION("char*")
     {
         std::string Value = "<unassigned>";
-        AutoRTFM::ETransactionResult Result = AutoRTFM::Transact([&] ()
+        AutoRTFM::ETransactionResult Result = AutoRTFM::Transact([&]
             {
                 // Note: AutoRTFM::Open() is returning a const char*
                 std::string Got = AutoRTFM::Open([] { return "meow"; });
@@ -154,7 +154,7 @@ TEST_CASE("Open.ReturnValue")
     {
         int Int = 0;
         std::string String = "<unassigned>";
-        AutoRTFM::ETransactionResult Result = AutoRTFM::Transact([&] ()
+        AutoRTFM::ETransactionResult Result = AutoRTFM::Transact([&]
             {
                 auto [I, S] = AutoRTFM::Open([] 
                     { 
@@ -241,6 +241,11 @@ TEST_CASE("Open.Collision")
 {
     AUTORTFM_SCOPED_ENABLE_MEMORY_VALIDATION_AS_WARNING();
 
+    struct FLargeStruct
+    {
+        int V[256];
+    };
+
     AutoRTFMTestUtils::FCaptureWarningContext WarningContext;
     int I = 0;
     int J = 0;
@@ -249,10 +254,10 @@ TEST_CASE("Open.Collision")
     {
         SECTION("Different Memory Locations")
         {
-            AutoRTFM::Transact([&] ()
+            AutoRTFM::Transact([&]
             {
                 I = 42;
-                AutoRTFM::Open([&] () { J = 24; });
+                AutoRTFM::Open([&] { J = 24; });
             });
 
             REQUIRE(I == 42);
@@ -262,13 +267,13 @@ TEST_CASE("Open.Collision")
 
         SECTION("Transact(Open(RecordOpenWriteNoMemoryValidation()), Open(Assign))")
         {
-            AutoRTFM::Transact([&] ()
+            AutoRTFM::Transact([&]
             {
-                AutoRTFM::Open([&] () 
+                AutoRTFM::Open([&]
                 {
                     AutoRTFM::RecordOpenWriteNoMemoryValidation(&I);
                 });
-                AutoRTFM::Open([&] () 
+                AutoRTFM::Open([&]
                 {
                     I = 42;
                 });
@@ -278,11 +283,58 @@ TEST_CASE("Open.Collision")
             REQUIRE(WarningContext.GetWarnings().empty());
         }
 
+        SECTION("Transact(Transact(Open(RecordOpenWriteNoMemoryValidation())), Open(Assign))")
+        {
+            AutoRTFM::Transact([&]
+            {
+                AutoRTFM::Transact([&]
+                {
+                    AutoRTFM::Open([&]
+                    {
+                        AutoRTFM::RecordOpenWriteNoMemoryValidation(&I);
+                    });
+                });
+                
+                AutoRTFM::Open([&]
+                {
+                    I = 42;
+                });
+            });
+
+            REQUIRE(I == 42);
+            REQUIRE(WarningContext.GetWarnings().empty());
+        }
+
+        SECTION("Transact(Open(RecordOpenWriteNoMemoryValidation()), Open(Assign)) <large>")
+        {
+            FLargeStruct LargeStruct;
+            AutoRTFM::Transact([&]
+            {
+                AutoRTFM::Open([&]
+                {
+                    AutoRTFM::RecordOpenWriteNoMemoryValidation(&LargeStruct);
+                });
+                AutoRTFM::Open([&]
+                {
+                    for (int& V : LargeStruct.V)
+                    {
+                        V = 42;
+                    }
+                });
+            });
+
+            for (int V : LargeStruct.V)
+            {
+                REQUIRE(V == 42);
+            }
+            REQUIRE(WarningContext.GetWarnings().empty());
+        }
+
         SECTION("Transact(Open(Assign), Assign)")
         {
-            AutoRTFM::Transact([&] ()
+            AutoRTFM::Transact([&]
             {
-                AutoRTFM::Open([&] () { I = 24; });
+                AutoRTFM::Open([&] { I = 24; });
                 I = 42;
             });
 
@@ -292,12 +344,12 @@ TEST_CASE("Open.Collision")
 
         SECTION("Transact(Assign, Open(Transact(Assign)))")
         {
-            AutoRTFM::Transact([&] ()
+            AutoRTFM::Transact([&]
             {
                 I = 24;
-                AutoRTFM::Open([&] ()
+                AutoRTFM::Open([&]
                 {
-                    AutoRTFM::Transact([&] ()
+                    AutoRTFM::Transact([&]
                     {
                         I = 42;
                     });
@@ -310,11 +362,11 @@ TEST_CASE("Open.Collision")
 
         SECTION("Transact(Open(Transact(Assign)), Assign)")
         {
-            AutoRTFM::Transact([&] ()
+            AutoRTFM::Transact([&]
             {
-                AutoRTFM::Open([&] ()
+                AutoRTFM::Open([&]
                 {
-                    AutoRTFM::Transact([&] ()
+                    AutoRTFM::Transact([&]
                     {
                         I = 24;
                     });
@@ -329,10 +381,10 @@ TEST_CASE("Open.Collision")
 
     SECTION("Transact(Assign, Open(Assign))")
     {
-        AutoRTFM::Transact([&] ()
+        AutoRTFM::Transact([&]
         {
             I = 42;
-            AutoRTFM::Open([&] () { I = 24; });
+            AutoRTFM::Open([&] { I = 24; });
         });
 
         REQUIRE_THAT(WarningContext.GetWarnings(),
@@ -341,12 +393,12 @@ TEST_CASE("Open.Collision")
 
     SECTION("Transact(Assign, Transact(Open(Assign)))")
     {
-        AutoRTFM::Transact([&] ()
+        AutoRTFM::Transact([&]
         {
             I = 42;
-            AutoRTFM::Transact([&] ()
+            AutoRTFM::Transact([&]
             {
-                AutoRTFM::Open([&] () { I = 24; });
+                AutoRTFM::Open([&] { I = 24; });
             });
         });
 
@@ -356,14 +408,14 @@ TEST_CASE("Open.Collision")
 
     SECTION("Transact(Assign, Transact(Transact(Open(Assign))))")
     {
-        AutoRTFM::Transact([&] ()
+        AutoRTFM::Transact([&]
         {
             I = 42;
-            AutoRTFM::Transact([&] ()
+            AutoRTFM::Transact([&]
             {
-                AutoRTFM::Transact([&] ()
+                AutoRTFM::Transact([&]
                 {
-                    AutoRTFM::Open([&] ()
+                    AutoRTFM::Open([&]
                     {
                         I = 24;
                     });
@@ -377,14 +429,14 @@ TEST_CASE("Open.Collision")
 
     SECTION("Transact(Transact(Assign, Transact(Open(Assign))))")
     {
-        AutoRTFM::Transact([&] ()
+        AutoRTFM::Transact([&]
         {
-            AutoRTFM::Transact([&] ()
+            AutoRTFM::Transact([&]
             {
                 I = 42;
-                AutoRTFM::Transact([&] ()
+                AutoRTFM::Transact([&]
                 {
-                    AutoRTFM::Open([&] ()
+                    AutoRTFM::Open([&]
                     {
                         I = 24;
                     });
@@ -398,13 +450,13 @@ TEST_CASE("Open.Collision")
 
     SECTION("Transact(Transact(Assign, Open(Assign, Transact())))")
     {
-        AutoRTFM::Transact([&] ()
+        AutoRTFM::Transact([&]
         {
             I = 42;
-            AutoRTFM::Open([&] ()
+            AutoRTFM::Open([&]
             {
                 I = 24;
-                AutoRTFM::Transact([&] () {});
+                AutoRTFM::Transact([&] {});
             });
         });
 
@@ -414,12 +466,12 @@ TEST_CASE("Open.Collision")
 
     SECTION("Transact(Transact(Assign, Open(Transact(), Assign)))")
     {
-        AutoRTFM::Transact([&] ()
+        AutoRTFM::Transact([&]
         {
             I = 42;
-            AutoRTFM::Open([&] ()
+            AutoRTFM::Open([&]
             {
-                AutoRTFM::Transact([&] () {});
+                AutoRTFM::Transact([&] {});
                 I = 24;
             });
         });
@@ -430,7 +482,7 @@ TEST_CASE("Open.Collision")
 
     SECTION("Transact(Assign, CallOpen(Assign))")
     {
-        AutoRTFM::Transact([&] ()
+        AutoRTFM::Transact([&]
         {
             I = 42;
             AssignIntPointer(&I, 24);
@@ -442,10 +494,10 @@ TEST_CASE("Open.Collision")
 
     SECTION("Transact(Assign, Transact(CallOpen(Assign)))")
     {
-        AutoRTFM::Transact([&] ()
+        AutoRTFM::Transact([&]
         {
             I = 42;
-            AutoRTFM::Transact([&] ()
+            AutoRTFM::Transact([&]
             {
                 AssignIntPointer(&I, 24);
             });
@@ -457,7 +509,7 @@ TEST_CASE("Open.Collision")
 
     SECTION("Transact(Assign, OpenNoValidation(Assign))")
     {
-        AutoRTFM::Transact([&] ()
+        AutoRTFM::Transact([&]
         {
             I = 42;
             AutoRTFM::Open<AutoRTFM::EMemoryValidationLevel::Disabled>([&]
@@ -480,7 +532,7 @@ TEST_CASE("Open.Collision")
             }
         };
 
-        AutoRTFM::Transact([&] ()
+        AutoRTFM::Transact([&]
         {
             I = 42;
             S::AssignToInt(I);
@@ -491,7 +543,7 @@ TEST_CASE("Open.Collision")
 
     SECTION("Transact(Assign, Open(Assign, OpenNoValidation()))")
     {
-        AutoRTFM::Transact([&] ()
+        AutoRTFM::Transact([&]
         {
             I = 42;
             AutoRTFM::Open([&]
@@ -507,7 +559,7 @@ TEST_CASE("Open.Collision")
     
     SECTION("Transact(Assign, Open(Assign, Close(OpenNoValidation())))")
     {
-        AutoRTFM::Transact([&] ()
+        AutoRTFM::Transact([&]
         {
             I = 42;
             AutoRTFM::Open([&]
@@ -526,13 +578,13 @@ TEST_CASE("Open.Collision")
 
     SECTION("Transact(Open(RecordOpenWrite()), Open(Assign))")
     {
-        AutoRTFM::Transact([&] ()
+        AutoRTFM::Transact([&]
         {
-            AutoRTFM::Open([&] () 
+            AutoRTFM::Open([&]
             {
                 AutoRTFM::RecordOpenWrite(&I);
             });
-            AutoRTFM::Open([&] () 
+            AutoRTFM::Open([&]
             {
                 I = 42;
             });
@@ -545,17 +597,17 @@ TEST_CASE("Open.Collision")
 
     SECTION("Transact(Open(RecordOpenWriteNoMemoryValidation()), Open(RecordOpenWrite()), Open(Assign))")
     {
-        AutoRTFM::Transact([&] ()
+        AutoRTFM::Transact([&]
         {
-            AutoRTFM::Open([&] () 
+            AutoRTFM::Open([&]
             {
                 AutoRTFM::RecordOpenWriteNoMemoryValidation(&I);
             });
-            AutoRTFM::Open([&] () 
+            AutoRTFM::Open([&]
             {
                 AutoRTFM::RecordOpenWrite(&I);
             });
-            AutoRTFM::Open([&] () 
+            AutoRTFM::Open([&]
             {
                 I = 42;
             });
@@ -565,4 +617,51 @@ TEST_CASE("Open.Collision")
             Catch::Matchers::VectorContains(FString(AutoRTFMTestUtils::kMemoryModifiedWarning)));
     }
 
+    SECTION("Transact(Transact(Open(RecordOpenWriteNoMemoryValidation())), Open(RecordOpenWrite()), Open(Assign))")
+    {
+        AutoRTFM::Transact([&]
+        {
+            AutoRTFM::Transact([&]
+            {
+                AutoRTFM::Open([&]
+                {
+                    AutoRTFM::RecordOpenWriteNoMemoryValidation(&I);
+                });
+            });
+            AutoRTFM::Open([&]
+            {
+                AutoRTFM::RecordOpenWrite(&I);
+            });
+            AutoRTFM::Open([&]
+            {
+                I = 42;
+            });
+        });
+
+        REQUIRE_THAT(WarningContext.GetWarnings(),
+            Catch::Matchers::VectorContains(FString(AutoRTFMTestUtils::kMemoryModifiedWarning)));
+    }
+
+    SECTION("Transact(Open(RecordOpenWriteNoMemoryValidation()), Open(RecordOpenWrite()), Open(Assign)) <large>")
+    {
+        FLargeStruct LargeStruct;
+        AutoRTFM::Transact([&]
+        {
+            AutoRTFM::Open([&]
+            {
+                AutoRTFM::RecordOpenWriteNoMemoryValidation(&LargeStruct);
+            });
+            AutoRTFM::Open([&]
+            {
+                AutoRTFM::RecordOpenWrite(&LargeStruct);
+            });
+            AutoRTFM::Open([&]
+            {
+                LargeStruct.V[255] = 42;
+            });
+        });
+
+        REQUIRE_THAT(WarningContext.GetWarnings(),
+            Catch::Matchers::VectorContains(FString(AutoRTFMTestUtils::kMemoryModifiedWarning)));
+    }
 }
diff --git a/Engine/Source/Runtime/Core/Private/AutoRTFM/Transaction.cpp b/Engine/Source/Runtime/Core/Private/AutoRTFM/Transaction.cpp
--- a/Engine/Source/Runtime/Core/Private/AutoRTFM/Transaction.cpp
+++ b/Engine/Source/Runtime/Core/Private/AutoRTFM/Transaction.cpp
@@ -234,6 +234,7 @@ void FTransaction::CommitNested()
 
 		FHitSet::Key HitSetEntry(Write.GetOriginal());
 		HitSetEntry.TopTag().Size = static_cast<uint16_t>(Write.GetSize());
+		HitSetEntry.TopTag().bNoMemoryValidation = Write.IgnoreForValidation();
 
         Parent->HitSet.FindOrTryInsert(HitSetEntry);
     }
diff --git a/Engine/Source/Runtime/Core/Private/AutoRTFM/Transaction.h b/Engine/Source/Runtime/Core/Private/AutoRTFM/Transaction.h
--- a/Engine/Source/Runtime/Core/Private/AutoRTFM/Transaction.h
+++ b/Engine/Source/Runtime/Core/Private/AutoRTFM/Transaction.h
@@ -119,7 +119,7 @@ public:
 
     // Record that a write is about to occur at the given LogicalAddress of Size bytes.
     void RecordWrite(void* LogicalAddress, size_t Size, bool bNoMemoryValidation = false);
-    void RecordWriteMaxPageSized(void* LogicalAddress, size_t Size, bool bNoMemoryValidation = false);
+    void RecordWriteMaxPageSized(void* LogicalAddress, size_t Size, bool bNoMemoryValidation);
     template<unsigned SIZE> void RecordWrite(void* LogicalAddress);
 	template<unsigned SIZE> void RecordWriteInsertedSlow(void* LogicalAddress);
 	template<unsigned SIZE> void RecordWriteNotInsertedSlow(void* LogicalAddress);
diff --git a/Engine/Source/Runtime/Core/Private/AutoRTFM/TransactionInlines.h b/Engine/Source/Runtime/Core/Private/AutoRTFM/TransactionInlines.h
--- a/Engine/Source/Runtime/Core/Private/AutoRTFM/TransactionInlines.h
+++ b/Engine/Source/Runtime/Core/Private/AutoRTFM/TransactionInlines.h
@@ -98,7 +98,7 @@ UE_AUTORTFM_FORCEINLINE bool FTransaction::ShouldRecordWrite(void* LogicalAddres
 }
 
 AUTORTFM_NO_ASAN UE_AUTORTFM_FORCEINLINE
-void FTransaction::RecordWriteMaxPageSized(void* LogicalAddress, size_t Size, bool bNoMemoryValidation /* = false */)
+void FTransaction::RecordWriteMaxPageSized(void* LogicalAddress, size_t Size, bool bNoMemoryValidation)
 {
     void* CopyAddress = WriteLogBumpAllocator.Allocate(Size);
     memcpy(CopyAddress, LogicalAddress, Size);
@@ -159,7 +159,7 @@ void FTransaction::RecordWrite(void* LogicalAddress, size_t Size, bool bNoMemory
 
     for (; (I + FWriteLogBumpAllocator::MaxSize) < Size; I += FWriteLogBumpAllocator::MaxSize)
     {
-        RecordWriteMaxPageSized(Address + I, FWriteLogBumpAllocator::MaxSize);
+        RecordWriteMaxPageSized(Address + I, FWriteLogBumpAllocator::MaxSize, bNoMemoryValidation);
     }
 
     // Remainder at the end of the memcpy.
