[AutoRTFMTests] Zero-initialize structure used in tests.

The reuse of initialized data can (and does) cause test failures in certain configurations.

diff --git a/Engine/Source/Programs/AutoRTFMTests/Private/Open.cpp b/Engine/Source/Programs/AutoRTFMTests/Private/Open.cpp
--- a/Engine/Source/Programs/AutoRTFMTests/Private/Open.cpp
+++ b/Engine/Source/Programs/AutoRTFMTests/Private/Open.cpp
@@ -307,7 +307,7 @@ TEST_CASE("Open.Collision")
 
         SECTION("Transact(Open(RecordOpenWriteNoMemoryValidation()), Open(Assign)) <large>")
         {
-            FLargeStruct LargeStruct;
+            FLargeStruct LargeStruct{};
             AutoRTFM::Transact([&]
             {
                 AutoRTFM::Open([&]
@@ -644,7 +644,7 @@ TEST_CASE("Open.Collision")
 
     SECTION("Transact(Open(RecordOpenWriteNoMemoryValidation()), Open(RecordOpenWrite()), Open(Assign)) <large>")
     {
-        FLargeStruct LargeStruct;
+        FLargeStruct LargeStruct{};
         AutoRTFM::Transact([&]
         {
             AutoRTFM::Open([&]
