Add nightly AutoRTFM LTO tests to catch bugs with all the fun LTO causes.

#rnx

diff --git a/Engine/Source/Programs/AutoRTFMTests/AutoRTFMTests.Target.cs b/Engine/Source/Programs/AutoRTFMTests/AutoRTFMTests.Target.cs
--- a/Engine/Source/Programs/AutoRTFMTests/AutoRTFMTests.Target.cs
+++ b/Engine/Source/Programs/AutoRTFMTests/AutoRTFMTests.Target.cs
@@ -3,6 +3,7 @@
 using EpicGames.Core;
 using UnrealBuildTool;
 using System.Collections.Generic;
+using UnrealBuildBase;
 
 [SupportedPlatforms(UnrealPlatformClass.All)]
 public class AutoRTFMTestsTarget : TargetRules
@@ -13,11 +14,11 @@ public class AutoRTFMTestsTarget : TargetRules
 
 	public AutoRTFMTestsTarget(TargetInfo Target) : base(Target)
 	{
+		DefaultBuildSettings = BuildSettingsVersion.Latest;
 		Type = TargetType.Program;
 		LinkType = TargetLinkType.Monolithic;
 		IncludeOrderVersion = EngineIncludeOrderVersion.Latest;
 
-		Name = "AutoRTFMTests";
 		LaunchModuleName = "AutoRTFMTests";
 
 		bCompileAgainstEngine = false;
