[AutoRTFM] Add a test verifying that lambdas returned from an ALWAYS_OPEN function are still transactionally safe.

diff --git a/Engine/Source/Programs/AutoRTFMTests/Private/Open.cpp b/Engine/Source/Programs/AutoRTFMTests/Private/Open.cpp
--- a/Engine/Source/Programs/AutoRTFMTests/Private/Open.cpp
+++ b/Engine/Source/Programs/AutoRTFMTests/Private/Open.cpp
@@ -16,6 +16,17 @@ UE_AUTORTFM_ALWAYS_OPEN static void AssignIntPointer(int* Pointer, int Value)
     *Pointer = Value;
 }
 
+using FnPtr = void (*)(int* Pointer);
+
+UE_AUTORTFM_ALWAYS_OPEN static FnPtr GetFunctionPointer()
+{
+	auto FnPtr = +[](int* Pointer)
+	{
+		*Pointer = 42;
+	};
+	return FnPtr;
+}
+
 }
 
 TEST_CASE("Open")
@@ -115,6 +126,45 @@ TEST_CASE("Open.Atomics")
     REQUIRE(bDidRun);
 }
 
+TEST_CASE("Open.FunctionPtrFromAlwaysOpenFunction")
+{
+	FnPtr Func = GetFunctionPointer();
+
+	int Value = 0;
+
+	SECTION("Commit")
+	{
+		AutoRTFM::Commit([&]
+		{
+			Func(&Value);
+			REQUIRE(Value == 42);
+		});
+		REQUIRE(Value == 42);
+	}
+
+	SECTION("Abort")
+	{
+		AutoRTFM::ETransactionResult Result = AutoRTFM::Transact([&]
+		{
+			Func(&Value);
+			REQUIRE(Value == 42);
+			AutoRTFM::AbortTransaction();
+		});
+		REQUIRE(Result == AutoRTFM::ETransactionResult::AbortedByRequest);
+		REQUIRE(Value == 0);
+	}
+
+	SECTION("Open")
+	{
+		AutoRTFM::Open([&]
+		{
+			Func(&Value);
+			REQUIRE(Value == 42);
+		});
+		REQUIRE(Value == 42);
+	}
+}
+
 TEST_CASE("Open.ReturnValue")
 {
     static_assert(AutoRTFM::IsSafeToReturnFromOpen<int>);
