Remove the open + on-abort on the invocation lock list of multicast delegate. These aren't actually atomic (there is another thing that prevents the atomic problems in multicast delegate), so this should be safe to do. Doing this reduces the average number of on-aborts in a transaction by 2.9x.

diff --git a/Engine/Source/Programs/AutoRTFMTests/Private/Delegate.cpp b/Engine/Source/Programs/AutoRTFMTests/Private/Delegate.cpp
--- a/Engine/Source/Programs/AutoRTFMTests/Private/Delegate.cpp
+++ b/Engine/Source/Programs/AutoRTFMTests/Private/Delegate.cpp
@@ -1,7 +1,7 @@
 // Copyright Epic Games, Inc. All Rights Reserved.
 
 #include "Catch2Includes.h"
-#include <AutoRTFM/AutoRTFM.h>
+#include <AutoRTFM/AutoRTFMTesting.h>
 #include "Delegates/Delegate.h"
 #include "Utils.h"
 
@@ -115,3 +115,226 @@ TEST_CASE("Delegate.Broadcast")
 		REQUIRE(0 == CallCount);
 	}
 }
+
+TEST_CASE("Delegate.Recursive")
+{
+	SECTION("Open <-> Close")
+	{
+		TMulticastDelegate<void(int, int&)> Delegate;
+
+		Delegate.AddLambda([&](int X, int& R)
+		{
+			REQUIRE(AutoRTFM::IsTransactional());
+
+			if (AutoRTFM::IsClosed())
+			{
+				AutoRTFM::Open([&]
+				{
+					if (X <= 1)
+					{
+						R = X;
+					}
+					else
+					{
+						int R1, R2;
+						Delegate.Broadcast(X - 1, R1);
+						Delegate.Broadcast(X - 2, R2);
+						R = R1 + R2;
+					}
+				});
+			}
+			else
+			{
+				(void)AutoRTFM::Close([&]
+				{
+					if (X <= 1)
+					{
+						R = X;
+					}
+					else
+					{
+						int R1, R2;
+						Delegate.Broadcast(X - 1, R1);
+						Delegate.Broadcast(X - 2, R2);
+						R = R1 + R2;
+					}
+				});
+			}
+		});
+
+		AutoRTFM::Testing::Commit([&]
+		{
+			int R = 0;
+			Delegate.Broadcast(6, R);
+			REQUIRE(R == 8);
+		});
+
+		Delegate.AddStatic([](int, int&) {});
+	}
+
+	
+	SECTION("Open <-> Close With Abort")
+	{
+		TMulticastDelegate<void(int, int&, bool&)> Delegate;
+		TMulticastDelegate<void(int, int&, bool&)> Other;
+
+		Delegate.AddLambda([&](int X, int& R, bool& bAborting)
+		{
+			REQUIRE(AutoRTFM::IsTransactional());
+
+			if (X <= 1)
+			{
+				R = X;
+
+				const AutoRTFM::EContextStatus Status = AutoRTFM::Close([&]
+				{
+					AutoRTFM::AbortTransaction();	
+				});
+				
+				bAborting = Status != AutoRTFM::EContextStatus::OnTrack;
+			}
+			else if (AutoRTFM::IsClosed())
+			{
+				AutoRTFM::Open([&]
+				{
+					int R1, R2;
+					Delegate.Broadcast(X - 1, R1, bAborting);
+
+					if (bAborting)
+					{
+						return;
+					}
+
+					Delegate.Broadcast(X - 2, R2, bAborting);
+
+					if (bAborting)
+					{
+						return;
+					}
+
+					R = R1 + R2;
+				});
+			}
+			else
+			{
+				const AutoRTFM::EContextStatus Status = AutoRTFM::Close([&]
+				{
+					int R1, R2;
+					Delegate.Broadcast(X - 1, R1, bAborting);
+					Delegate.Broadcast(X - 2, R2, bAborting);
+					R = R1 + R2;
+				});
+
+				bAborting = Status != AutoRTFM::EContextStatus::OnTrack;
+			}
+		});
+
+		AutoRTFM::Testing::Abort([&]
+		{
+			bool bAborting = false;
+			int R = 0;
+			Delegate.Broadcast(6, R, bAborting);
+			FAIL("Unreachable!");
+		});
+
+		AutoRTFM::Testing::Abort([&]
+		{
+			bool bAborting = false;
+			int R = 0;
+			Delegate.Broadcast(6, R, bAborting);
+			Other = MoveTemp(Delegate);
+			FAIL("Unreachable!");
+		});
+
+		AutoRTFM::Testing::Abort([&]
+		{
+			bool bAborting = false;
+			int R = 0;
+			Delegate.Broadcast(5, R, bAborting);
+			FAIL("Unreachable!");
+		});
+
+		AutoRTFM::Testing::Abort([&]
+		{
+			bool bAborting = false;
+			int R = 0;
+			Delegate.Broadcast(5, R, bAborting);
+			Other = MoveTemp(Delegate);
+			FAIL("Unreachable!");
+		});
+
+		Other = MoveTemp(Delegate);
+	}
+}
+
+TEST_CASE("Delegate.CauseCompaction")
+{
+	// Tests are sensitive to retries. Disable for these tests.
+	AutoRTFMTestUtils::FScopedRetry Retry(AutoRTFM::ForTheRuntime::EAutoRTFMRetryTransactionState::NoRetry);
+
+	SECTION("Transact")
+	{
+		TMulticastDelegate<void()> Delegate;
+
+		FDelegateHandle Handle = Delegate.AddStatic([]() { FAIL("Unreachable!"); });
+
+		bool bHitOnce = false;
+
+		Delegate.AddLambda([&]()
+		{
+			if (!bHitOnce)
+			{
+				REQUIRE(Delegate.Remove(Handle));
+				bHitOnce = true;
+			}
+		});
+
+		AutoRTFM::Testing::Commit([&] { Delegate.Broadcast(); });
+
+		Delegate.Broadcast();
+	}
+
+	SECTION("OnCommit")
+	{
+		TMulticastDelegate<void()> Delegate;
+
+		FDelegateHandle Handle = Delegate.AddStatic([]() { FAIL("Unreachable!"); });
+
+		bool bHitOnce = false;
+
+		Delegate.AddLambda([&]()
+		{
+			if (!bHitOnce)
+			{
+				REQUIRE(Delegate.Remove(Handle));
+				bHitOnce = true;
+			}
+		});
+
+		AutoRTFM::Testing::Commit([&] { AutoRTFM::OnCommit([&] { Delegate.Broadcast(); }); });
+
+		Delegate.Broadcast();
+	}
+	
+	SECTION("OnAbort")
+	{
+		TMulticastDelegate<void()> Delegate;
+
+		FDelegateHandle Handle = Delegate.AddStatic([]() { FAIL("Unreachable!"); });
+
+		bool bHitOnce = false;
+
+		Delegate.AddLambda([&]()
+		{
+			if (!bHitOnce)
+			{
+				REQUIRE(Delegate.Remove(Handle));
+				bHitOnce = true;
+			}
+		});
+
+		AutoRTFM::Testing::Abort([&] { AutoRTFM::OnAbort([&] { Delegate.Broadcast(); }); AutoRTFM::AbortTransaction(); });
+
+		Delegate.Broadcast();
+	}
+}
diff --git a/Engine/Source/Runtime/Core/Public/Delegates/MulticastDelegateBase.h b/Engine/Source/Runtime/Core/Public/Delegates/MulticastDelegateBase.h
--- a/Engine/Source/Runtime/Core/Public/Delegates/MulticastDelegateBase.h
+++ b/Engine/Source/Runtime/Core/Public/Delegates/MulticastDelegateBase.h
@@ -317,6 +317,16 @@ protected:
 	}
 
 private:
+
+	inline bool ShouldSkipCompactionUnderAutoRTFM() const
+	{
+		// As a performance optimization for the AutoRTFM runtime, we *always*
+		// skip compacting a multicast delegate when called via the AutoRTFM
+		// runtime. We skip calls both in the bodies of transactions themselves,
+		// and also those that are called as part of on-commit or on-abort calls.
+		return AutoRTFM::IsTransactional() || AutoRTFM::IsCommittingOrAborting();
+	}
+
 	/**
 	 * Removes any expired or deleted functions from the invocation list.
 	 *
@@ -324,6 +334,11 @@ private:
 	 */
 	void CompactInvocationList(bool CheckThreshold=false)
 	{
+		if (ShouldSkipCompactionUnderAutoRTFM())
+		{
+			return;
+		}
+
 		// if locked and no object, just return
 		if (InvocationListLockCount > 0)
 		{
@@ -382,17 +397,23 @@ private:
 	/** Increments the lock counter for the invocation list. */
 	inline void LockInvocationList( ) const
 	{
-		// Validation disabled as the only closed write is the ctor.
-		UE_AUTORTFM_OPEN_NO_VALIDATION { ++InvocationListLockCount; };
-		UE_AUTORTFM_ONABORT(this) { --InvocationListLockCount; };
+		if (ShouldSkipCompactionUnderAutoRTFM())
+		{
+			return;
+		}
+		
+		++InvocationListLockCount;
 	}
 
 	/** Decrements the lock counter for the invocation list. */
 	inline void UnlockInvocationList( ) const
 	{
-		// Validation disabled as the only closed write is the ctor.
-		UE_AUTORTFM_OPEN_NO_VALIDATION { --InvocationListLockCount; };
-		UE_AUTORTFM_ONABORT(this) { ++InvocationListLockCount; };
+		if (ShouldSkipCompactionUnderAutoRTFM())
+		{
+			return;
+		}
+		
+		--InvocationListLockCount;
 	}
 
 	/** Returns the lock counter for the invocation list. */
