Fix the transactionally safe mutex to correctly work under on-commit/on-abort too, and remove an ensure which would trigger when trying to complete a transaction that had a mutex that was explicitly memset(0) (EG. via UObject creation/deletion).

diff --git a/Engine/Source/Programs/AutoRTFMTests/Private/UECoreTests.cpp b/Engine/Source/Programs/AutoRTFMTests/Private/UECoreTests.cpp
--- a/Engine/Source/Programs/AutoRTFMTests/Private/UECoreTests.cpp
+++ b/Engine/Source/Programs/AutoRTFMTests/Private/UECoreTests.cpp
@@ -2085,6 +2085,252 @@ TEST_CASE("UECore.TransactionallySafeMutex")
 			Mutex.Unlock();
 		}
 	}
+
+	SECTION("In On-Commit")
+	{
+		UE::FTransactionallySafeMutex Mutex;
+
+		AutoRTFM::Testing::Commit([&]
+		{
+			Mutex.Lock();
+
+			AutoRTFM::OnCommit([&]
+			{
+				REQUIRE(Mutex.IsLocked());
+				Mutex.Unlock();
+				REQUIRE(!Mutex.IsLocked());
+			});
+		});
+
+		AutoRTFM::Testing::Commit([&]
+		{
+			AutoRTFM::OnCommit([&]
+			{
+				REQUIRE(Mutex.IsLocked());
+				Mutex.Unlock();
+				REQUIRE(!Mutex.IsLocked());
+			});
+			
+			Mutex.Lock();
+		});
+
+		AutoRTFM::Testing::Commit([&]
+		{
+			REQUIRE(Mutex.TryLock());
+
+			AutoRTFM::OnCommit([&]
+			{
+				REQUIRE(Mutex.IsLocked());
+				Mutex.Unlock();
+				REQUIRE(!Mutex.IsLocked());
+			});
+		});
+
+		AutoRTFM::Testing::Commit([&]
+		{
+			AutoRTFM::OnCommit([&]
+			{
+				REQUIRE(Mutex.IsLocked());
+				Mutex.Unlock();
+				REQUIRE(!Mutex.IsLocked());
+			});
+			
+			REQUIRE(Mutex.TryLock());
+		});
+
+		Mutex.Lock();
+		REQUIRE(Mutex.IsLocked());
+
+		AutoRTFM::Testing::Commit([&]
+		{
+			AutoRTFM::OnCommit([&]
+			{
+				REQUIRE(!Mutex.IsLocked());
+				Mutex.Lock();
+				REQUIRE(Mutex.IsLocked());
+			});
+			
+			Mutex.Unlock();
+		});
+
+		REQUIRE(Mutex.IsLocked());
+
+		AutoRTFM::Testing::Commit([&]
+		{
+			Mutex.Unlock();
+
+			AutoRTFM::OnCommit([&]
+			{
+				REQUIRE(!Mutex.IsLocked());
+				Mutex.Lock();
+				REQUIRE(Mutex.IsLocked());
+			});
+		});
+
+		REQUIRE(Mutex.IsLocked());
+		
+		AutoRTFM::Testing::Commit([&]
+		{
+			Mutex.Unlock();
+
+			AutoRTFM::OnCommit([&]
+			{
+				REQUIRE(Mutex.TryLock());
+			});
+		});
+
+		REQUIRE(Mutex.IsLocked());
+		
+		AutoRTFM::Testing::Commit([&]
+		{
+			AutoRTFM::OnCommit([&]
+			{
+				REQUIRE(Mutex.TryLock());
+			});
+
+			Mutex.Unlock();
+		});
+
+		REQUIRE(Mutex.IsLocked());
+	}
+
+	SECTION("In On-Abort")
+	{
+		UE::FTransactionallySafeMutex Mutex;
+
+		AutoRTFM::Testing::Abort([&]
+		{
+			Mutex.Lock();
+
+			AutoRTFM::OnAbort([&]
+			{
+				REQUIRE(Mutex.IsLocked());
+				Mutex.Unlock();
+				REQUIRE(!Mutex.IsLocked());
+			});
+
+			AutoRTFM::AbortTransaction();
+		});
+
+		AutoRTFM::Testing::Abort([&]
+		{
+			AutoRTFM::OnAbort([&]
+			{
+				REQUIRE(!Mutex.IsLocked());
+			});
+			
+			Mutex.Lock();
+
+			AutoRTFM::AbortTransaction();
+		});
+
+		AutoRTFM::Testing::Abort([&]
+		{
+			REQUIRE(Mutex.TryLock());
+
+			AutoRTFM::OnAbort([&]
+			{
+				REQUIRE(Mutex.IsLocked());
+				Mutex.Unlock();
+				REQUIRE(!Mutex.IsLocked());
+			});
+
+			AutoRTFM::AbortTransaction();
+		});
+
+		AutoRTFM::Testing::Abort([&]
+		{
+			AutoRTFM::OnAbort([&]
+			{
+				REQUIRE(!Mutex.IsLocked());
+			});
+			
+			REQUIRE(Mutex.TryLock());
+
+			AutoRTFM::AbortTransaction();
+		});
+
+		Mutex.Lock();
+		REQUIRE(Mutex.IsLocked());
+
+		AutoRTFM::Testing::Abort([&]
+		{
+			AutoRTFM::OnAbort([&]
+			{
+				REQUIRE(Mutex.IsLocked());
+			});
+			
+			Mutex.Unlock();
+
+			AutoRTFM::AbortTransaction();
+		});
+
+		REQUIRE(Mutex.IsLocked());
+
+		AutoRTFM::Testing::Abort([&]
+		{
+			Mutex.Unlock();
+
+			AutoRTFM::OnAbort([&]
+			{
+				REQUIRE(Mutex.IsLocked());
+			});
+
+			AutoRTFM::AbortTransaction();
+		});
+
+		REQUIRE(Mutex.IsLocked());
+		
+		AutoRTFM::Testing::Abort([&]
+		{
+			Mutex.Unlock();
+
+			AutoRTFM::OnAbort([&]
+			{
+				REQUIRE(!Mutex.TryLock());
+			});
+
+			AutoRTFM::AbortTransaction();
+		});
+
+		REQUIRE(Mutex.IsLocked());
+		
+		AutoRTFM::Testing::Abort([&]
+		{
+			AutoRTFM::OnAbort([&]
+			{
+				REQUIRE(!Mutex.TryLock());
+			});
+
+			Mutex.Unlock();
+
+			AutoRTFM::AbortTransaction();
+		});
+
+		REQUIRE(Mutex.IsLocked());
+	}
+
+	SECTION("Locked Mutex In Destructed Object")
+	{
+		struct MyStruct final
+		{
+			UE::FTransactionallySafeMutex Mutex;
+
+			~MyStruct()
+			{
+				memset(this, 0, sizeof(MyStruct));
+			}
+		};
+
+		std::unique_ptr<MyStruct> Mine(new MyStruct);
+
+		AutoRTFM::Testing::Commit([&]
+		{
+			Mine->Mutex.Lock();
+			Mine.reset();
+			REQUIRE(!Mine);
+		});
+	}
 }
 
 TEST_CASE("UECore.FTextFormatPatternDefinition")
diff --git a/Engine/Source/Runtime/Core/Public/Async/TransactionallySafeMutex.h b/Engine/Source/Runtime/Core/Public/Async/TransactionallySafeMutex.h
--- a/Engine/Source/Runtime/Core/Public/Async/TransactionallySafeMutex.h
+++ b/Engine/Source/Runtime/Core/Public/Async/TransactionallySafeMutex.h
@@ -59,6 +59,10 @@ public:
 				return false;
 			});
 		}
+		else if (AutoRTFM::IsCommittingOrAborting() && DoesTransactionHoldLock())
+		{
+			return !IsTransactionallyUnlocked();
+		}
 		else
 		{
 			return FExternalMutex(State).IsLocked();
@@ -95,6 +99,11 @@ public:
 				return true;
 			});
 		}
+		else if (AutoRTFM::IsCommittingOrAborting() && DoesTransactionHoldLock())
+		{
+			SetTransactionalLockMode(ETransactionalLockMode::TransactionallyLocked);
+			return true;
+		}
 		else
 		{
 			const bool bResult = FExternalMutex(State).TryLock();
@@ -125,6 +134,10 @@ public:
 				RecordPreviousStateAndSetTransactionalLockMode(ETransactionalLockMode::TransactionallyLocked);
 			});
 		}
+		else if (AutoRTFM::IsCommittingOrAborting() && DoesTransactionHoldLock())
+		{
+			SetTransactionalLockMode(ETransactionalLockMode::TransactionallyLocked);
+		}
 		else
 		{
 			FExternalMutex(State).Lock();
@@ -151,6 +164,10 @@ public:
 				RecordPreviousStateAndSetTransactionalLockMode(ETransactionalLockMode::TransactionallyUnlocked);
 			});
 		}
+		else if (AutoRTFM::IsCommittingOrAborting() && DoesTransactionHoldLock())
+		{
+			SetTransactionalLockMode(ETransactionalLockMode::TransactionallyUnlocked);
+		}
 		else
 		{
 			ensure(!DoesTransactionHoldLock());
@@ -207,12 +224,6 @@ private:
 
 		AutoRTFM::PushOnCommitHandler(this, [this]
 		{
-			// If we commit, we won't have undone any of our transactional writes to the
-			// lock. So just double check that we still hold the lock here (famous last
-			// words, but this is a temporary check just while we deploy the new mutex
-			// more widely. Once we've got a few runs around the block we can remove!).
-			ensure(DoesTransactionHoldLock());
-
 			const bool bIsTransactionallyUnlocked = IsTransactionallyUnlocked();
 
 			State.fetch_and(~(IsLockedWithinTransactionFlag | IsUnlockedWithinTransactionFlag), std::memory_order_release);
@@ -230,14 +241,8 @@ private:
 		TransactionallyUnlocked = true,
 	};
 
-	inline void RecordPreviousStateAndSetTransactionalLockMode(ETransactionalLockMode LockMode)
+	inline void SetTransactionalLockMode(ETransactionalLockMode LockMode)
 	{
-		// Record the state into the write log. If an abort happens we will non-atomically
-		// overwrite `State`, removing the transactional flags. This will be followed-up
-		// with an atomic write in the on-abort handler to actually do the unlock though,
-		// thus ensuring data integrity.
-		AutoRTFM::RecordOpenWriteNoMemoryValidation(&State);
-
 		if (LockMode == ETransactionalLockMode::TransactionallyLocked)
 		{
 			State |= IsLockedWithinTransactionFlag;
@@ -249,6 +254,17 @@ private:
 			State &= ~IsLockedWithinTransactionFlag;
 		}
 	}
+
+	inline void RecordPreviousStateAndSetTransactionalLockMode(ETransactionalLockMode LockMode)
+	{
+		// Record the state into the write log. If an abort happens we will non-atomically
+		// overwrite `State`, removing the transactional flags. This will be followed-up
+		// with an atomic write in the on-abort handler to actually do the unlock though,
+		// thus ensuring data integrity.
+		AutoRTFM::RecordOpenWriteNoMemoryValidation(&State);
+
+		SetTransactionalLockMode(LockMode);
+	}
 };
 
 } // UE
