Add a test case for `FPackageName::DoesPackageExistEx`.

#rnx

diff --git a/Engine/Source/Programs/AutoRTFMTests/Private/UECoreTests.cpp b/Engine/Source/Programs/AutoRTFMTests/Private/UECoreTests.cpp
--- a/Engine/Source/Programs/AutoRTFMTests/Private/UECoreTests.cpp
+++ b/Engine/Source/Programs/AutoRTFMTests/Private/UECoreTests.cpp
@@ -23,6 +23,7 @@
 #include "Misc/ScopeRWLock.h"
 #include "Misc/ConfigCacheIni.h"
 #include "Misc/CoreDelegates.h"
+#include "Misc/PackageName.h"
 #include "StructUtils/PropertyBag.h"
 #include "Utils.h"
 
@@ -3088,3 +3089,12 @@ TEST_CASE("UECore.CoreRedirects")
 		}
 	}
 }
+
+TEST_CASE("UECore.PackageName")
+{
+	AutoRTFM::Testing::Commit([&]
+	{
+		FPackagePath Path = FPackagePath::FromLocalPath(FString("/Fake/Package/Path.lol"));
+		REQUIRE(FPackageName::EPackageLocationFilter::None == FPackageName::DoesPackageExistEx(Path, FPackageName::EPackageLocationFilter::IoDispatcher));
+	});
+}
