Fix a test bug uncovered when compiling for debug, and add testing for that to nightly.

#rnx

diff --git a/Engine/Source/Programs/AutoRTFMTests/Private/Open.cpp b/Engine/Source/Programs/AutoRTFMTests/Private/Open.cpp
--- a/Engine/Source/Programs/AutoRTFMTests/Private/Open.cpp
+++ b/Engine/Source/Programs/AutoRTFMTests/Private/Open.cpp
@@ -187,6 +187,9 @@ TEST_CASE("Open.ReturnValue")
         REQUIRE(42 == Value);
     }
 
+#if !UE_BUILD_DEBUG
+	// jira SOL-7541: std::string only under debug+linux hits a stdlib
+	// missing function in the closed!
     SECTION("char*")
     {
         std::string Value = "<unassigned>";
@@ -199,6 +202,7 @@ TEST_CASE("Open.ReturnValue")
         REQUIRE(AutoRTFM::ETransactionResult::Committed == Result);
         REQUIRE("meow" == Value);
     }
+#endif //!UE_BUILD_DEBUG
     
     SECTION("tuple")
     {
diff --git a/Engine/Source/Programs/AutoRTFMTests/Private/UECoreTests.cpp b/Engine/Source/Programs/AutoRTFMTests/Private/UECoreTests.cpp
--- a/Engine/Source/Programs/AutoRTFMTests/Private/UECoreTests.cpp
+++ b/Engine/Source/Programs/AutoRTFMTests/Private/UECoreTests.cpp
@@ -2202,13 +2202,16 @@ TEST_CASE("UECore.TransactionallySafeMutex")
 
 		AutoRTFM::Testing::Abort([&]
 		{
+			AutoRTFM::OnAbort([&]
+			{
+				REQUIRE(!Mutex.IsLocked());
+			});
+
 			Mutex.Lock();
 
 			AutoRTFM::OnAbort([&]
 			{
 				REQUIRE(Mutex.IsLocked());
-				Mutex.Unlock();
-				REQUIRE(!Mutex.IsLocked());
 			});
 
 			AutoRTFM::AbortTransaction();
@@ -2221,20 +2224,11 @@ TEST_CASE("UECore.TransactionallySafeMutex")
 				REQUIRE(!Mutex.IsLocked());
 			});
 			
-			Mutex.Lock();
-
-			AutoRTFM::AbortTransaction();
-		});
-
-		AutoRTFM::Testing::Abort([&]
-		{
 			REQUIRE(Mutex.TryLock());
 
 			AutoRTFM::OnAbort([&]
 			{
 				REQUIRE(Mutex.IsLocked());
-				Mutex.Unlock();
-				REQUIRE(!Mutex.IsLocked());
 			});
 
 			AutoRTFM::AbortTransaction();
