[AutoRTFM] Add TTaskArray tests, fix bug found

These tests check the existing behavior and will be used to verify an optimization of TTaskArray.

Bug fix: ForwardProcessEntry() and BackwardProcessEntry() can nullify the void* value of the BitMasksPerKey map entry. Currently a transaction will only ever call these once per TTaskArray lifetime, but if called twice (same or either), then we could dereference a nullptr.

diff --git a/Engine/Source/Programs/AutoRTFMTests/Private/TaskArrayTests.cpp b/Engine/Source/Programs/AutoRTFMTests/Private/TaskArrayTests.cpp
new file mode 100644
--- /dev/null
+++ b/Engine/Source/Programs/AutoRTFMTests/Private/TaskArrayTests.cpp
@@ -0,0 +1,254 @@
+// Copyright Epic Games, Inc. All Rights Reserved.
+
+#include "AutoRTFM/TaskArray.h"
+#include "Catch2Includes.h"
+#include "Templates/UnrealTemplate.h"
+#include <memory>
+#include <random>
+#include <unordered_map>
+
+TEST_CASE("TTaskArray")
+{
+	std::mt19937 Rand(0x1234);
+	auto TaskArray = std::make_unique<AutoRTFM::TTaskArray<int>>();
+
+	struct FEntry
+	{
+		int Value = 0;
+		FEntry* Prev = nullptr;
+		FEntry* Next = nullptr;
+	};
+
+	std::vector<void*> Keys;
+	std::unordered_map<void*, std::vector<FEntry*>> EntriesByKey;
+	FEntry* Head = nullptr;
+	FEntry* Tail = nullptr;
+
+	auto DeleteEntry = [&](FEntry* Entry)
+	{
+		if (Head == Entry) { Head = Entry->Next; }
+		if (Tail == Entry) { Tail = Entry->Prev; }
+		if (Entry->Prev) { Entry->Prev->Next = Entry->Next; }
+		if (Entry->Next) { Entry->Next->Prev = Entry->Prev; }
+		Entry->Prev = nullptr;
+		Entry->Next = nullptr;
+		delete Entry;
+	};
+
+	struct FConfig
+	{
+		size_t NumKeys = 0;
+		size_t NumElements = 0;
+	};
+
+	auto Populate = [&](FConfig Config)
+	{
+		for (size_t I = 0; I < Config.NumKeys; I++)
+		{
+			void* Key = reinterpret_cast<void*>(static_cast<uintptr_t>(I + 1) << 20);
+			Keys.push_back(Key);
+		}
+		
+		for (size_t I = 0; I < Config.NumElements; I++)
+		{
+			const int Value = static_cast<int>(Rand());
+			FEntry* const Entry = new FEntry{};
+			Entry->Value = Value;
+			if (!Head)
+			{ 
+				Head = Entry;
+			}
+			else
+			{
+				Tail->Next = Entry;
+				Entry->Prev = Tail;
+			}
+			Tail = Entry;
+			if (Config.NumKeys > 0 && Rand() & 1)
+			{
+				void* Key = Keys[Rand() % Keys.size()];
+				EntriesByKey[Key].push_back(Entry);
+				TaskArray->AddKeyed(Key, Value);
+			}
+			else
+			{
+				TaskArray->Add(Value);
+			}
+		}
+	};
+	
+	auto Check = [&]
+	{
+		if (Head)
+		{
+			size_t Count = 0;
+			FEntry* Current = Head;
+			TaskArray->ForEachForward([&](int Got)
+			{
+				REQUIRE(Current);
+				REQUIRE(Current->Value == Got);
+				Current = Current->Next;
+				Count++;
+			});
+			// REQUIRE(Count == TaskArray->Num()); // Currently Num() includes removed elements
+			REQUIRE(!Current);
+			
+			Count = 0;
+			Current = Tail;
+			TaskArray->ForEachBackward([&](int Got)
+			{
+				REQUIRE(Current);
+				REQUIRE(Current->Value == Got);
+				Current = Current->Prev;
+				Count++;
+			});
+			// REQUIRE(Count == TaskArray->Num()); // Currently Num() includes removed elements
+			REQUIRE(!Current);
+			REQUIRE(!TaskArray->IsEmpty()); // Currently IsEmpty() includes removed elements
+		}
+		else
+		{
+			REQUIRE(TaskArray->Num() == 0);
+			TaskArray->ForEachForward([&](int)
+			{
+				FAIL("Should not be called");
+			});
+			TaskArray->ForEachBackward([&](int)
+			{
+				FAIL("Should not be called");
+			});
+			REQUIRE(TaskArray->IsEmpty()); // Currently IsEmpty() includes removed elements
+		}
+	};
+	
+	auto CallWithConfigs = [&](auto&& Fn)
+	{
+		auto WithKeys = [&](size_t NumKeys)
+		{
+			SECTION("NumElements: 0")
+			{
+				Fn(FConfig{NumKeys, /* NumElements */ 0});
+			}
+
+			SECTION("NumElements: 10")
+			{
+				Fn(FConfig{NumKeys, /* NumElements */ 10});
+			}
+
+			SECTION("NumElements: 100")
+			{
+				Fn(FConfig{NumKeys, /* NumElements */ 100});
+			}
+
+			SECTION("NumElements: 1000")
+			{
+				Fn(FConfig{NumKeys, /* NumElements */ 1000});
+			}
+
+			SECTION("NumElements: 10000")
+			{
+				Fn(FConfig{NumKeys, /* NumElements */ 10000});
+			}
+		};
+		
+		SECTION("NumKeys: 0")
+		{
+			WithKeys(/* NumKeys */ 0);
+		}
+
+		SECTION("NumKeys: 5")
+		{
+			WithKeys(/* NumKeys */ 5);
+		}
+
+		SECTION("NumKeys: 10")
+		{
+			WithKeys(/* NumKeys */ 10);
+		}
+	};
+
+	SECTION("Add")
+	{
+		CallWithConfigs([&](FConfig Config)
+		{
+			Populate(Config);
+			Check();
+		});
+	}
+
+	SECTION("DeleteKey")
+	{
+		CallWithConfigs([&](FConfig Config)
+		{
+			if (Config.NumKeys == 0)
+			{
+				return;
+			}
+
+			Populate(Config);
+
+			for (size_t I = 0; I < 10; I++)
+			{
+				void* Key = Keys[Rand() & Config.NumKeys];
+				TaskArray->DeleteKey(Key);
+				if (std::vector<FEntry*>& Entries = EntriesByKey[Key]; Entries.size())
+				{
+					DeleteEntry(Entries.back());
+					Entries.pop_back();
+				}
+			}
+
+			Check();
+		});
+	}
+
+	SECTION("DeleteAllMatchingKeys")
+	{
+		CallWithConfigs([&](FConfig Config)
+		{
+			if (Config.NumKeys == 0)
+			{
+				return;
+			}
+
+			Populate(Config);
+
+			for (size_t I = 0; I < 4; I++)
+			{
+				void* Key = Keys[Rand() & Config.NumKeys];
+				TaskArray->DeleteAllMatchingKeys(Key);
+				std::vector<FEntry*>& Entries = EntriesByKey[Key];
+				while (Entries.size())
+				{
+					DeleteEntry(Entries.back());
+					Entries.pop_back();
+				}
+			}
+
+			Check();
+		});
+	}
+
+	SECTION("DeleteAllMatchingKeys")
+	{
+		CallWithConfigs([&](FConfig Config)
+		{
+			Populate(Config);
+
+			auto Second = std::make_unique<AutoRTFM::TTaskArray<int>>();
+			std::swap(TaskArray, Second);
+			Populate(Config);
+			std::swap(TaskArray, Second);
+
+			TaskArray->AddAll(MoveTemp(*Second));
+
+			Check();
+		});
+	}
+
+	// Clean-up
+	while (Head)
+	{
+		DeleteEntry(Head);
+	}
+}
diff --git a/Engine/Source/Runtime/Core/Private/AutoRTFM/TaskArray.h b/Engine/Source/Runtime/Core/Private/AutoRTFM/TaskArray.h
--- a/Engine/Source/Runtime/Core/Private/AutoRTFM/TaskArray.h
+++ b/Engine/Source/Runtime/Core/Private/AutoRTFM/TaskArray.h
@@ -7,6 +7,7 @@
 #include "Containers/Map.h"
 #include "Templates/Function.h"
 #include "Utils.h"
+#include <type_traits>
 
 namespace AutoRTFM
 {
@@ -353,19 +354,25 @@ private:
 		{
 			FBitMaskPerKeyType* const KeyedElement = BitMasksPerKey.Find(*Pair.Key);
 
-			if (KeyedElement->Num() > 0)
+			if (KeyedElement && KeyedElement->Num() > 0)
 			{
 				const bool bFirst = KeyedElement->BitAt(0);
 				KeyedElement->RemoveFirst();
 
 				if (bFirst)
 				{
-					static const TFunction<void()> EmptyFunctor = [] {};
-
 					// We have a deferred signal to pop this element, so do it now.
 					*Pair.Key = SKey{};
-					Pair.Val->Reset();
-					*Pair.Val = EmptyFunctor;
+					if constexpr (std::is_invocable_v<InElementType>)
+					{
+						static const TFunction<void()> EmptyFunctor = [] {};
+						Pair.Val->Reset();
+						*Pair.Val = EmptyFunctor;
+					}
+					else
+					{
+						*Pair.Val = {};
+					}
 					return;
 				}
 			}
@@ -380,6 +387,10 @@ private:
 		if (Pair.Key)
 		{
 			FBitMaskPerKeyType* const KeyedElement = BitMasksPerKey.Find(*Pair.Key);
+			if (!KeyedElement)
+			{
+				return;
+			}
 
 			const int32 Num = KeyedElement->Num();
 
@@ -390,12 +401,19 @@ private:
 
 				if (bLast)
 				{
-					static const TFunction<void()> EmptyFunctor = [] {};
-
 					// We have a deferred signal to pop this element, so do it now.
 					*Pair.Key = SKey{};
-					Pair.Val->Reset();
-					*Pair.Val = EmptyFunctor;
+					if constexpr (std::is_invocable_v<InElementType>)
+					{
+						static const TFunction<void()> EmptyFunctor = [] {};
+						Pair.Val->Reset();
+						*Pair.Val = EmptyFunctor;
+					}
+					else
+					{
+						*Pair.Val = {};
+					}
+
 					return;
 				}
 			}
