[AutoRTFM] Allow UObject renames to be rolled back properly.

The existing code was tripping the open sanitizer in rare cases, and UObject::Rename in the closed was not getting rolled back after an AutoRTFM abort. I've removed an Open block from UObject::Rename and this appears to fix the transactional behavior.

Additionally, the UObject unit test code now uses the `AutoRTFM::Testing` methods where possible.

diff --git a/Engine/Plugins/Experimental/UIFramework/Source/Private/Types/UIFWidgetTree.cpp b/Engine/Plugins/Experimental/UIFramework/Source/Private/Types/UIFWidgetTree.cpp
--- a/Engine/Plugins/Experimental/UIFramework/Source/Private/Types/UIFWidgetTree.cpp
+++ b/Engine/Plugins/Experimental/UIFramework/Source/Private/Types/UIFWidgetTree.cpp
@@ -222,18 +222,7 @@ void FUIFrameworkWidgetTree::AuthorityAddChildInternal(UUIFrameworkWidget* Paren
 			//If the outer is the transient package, then there are no replication owner yet and it safe to rename it with the correct new outer.
 			//If the Outer is not the transient, then the widget got replicated with another player.There are no "reset" and we should duplicate and delete the object.
 			//For now only do the rename.It works but it is not the best.
-
-			UObject* OldOuter = Child->GetOuter();
-
-			UE_AUTORTFM_OPEN
-			{
-				Child->Rename(nullptr, ReplicatedOwner);
-			};
-
-			UE_AUTORTFM_ONABORT(Child, OldOuter)
-			{
-				Child->Rename(nullptr, OldOuter);
-			};
+ 			Child->Rename(nullptr, ReplicatedOwner);
 		}
 	}
 
diff --git a/Engine/Source/Programs/AutoRTFMTests/Private/UObjectTests.cpp b/Engine/Source/Programs/AutoRTFMTests/Private/UObjectTests.cpp
--- a/Engine/Source/Programs/AutoRTFMTests/Private/UObjectTests.cpp
+++ b/Engine/Source/Programs/AutoRTFMTests/Private/UObjectTests.cpp
@@ -1,5 +1,6 @@
 // Copyright Epic Games, Inc. All Rights Reserved.
 
+#include "AutoRTFM/AutoRTFMTesting.h"
 #include "AutoRTFM/ScopedGuard.h"
 #include "GenericPlatform/GenericPlatformMisc.h"
 #include "HAL/MallocLeakDetection.h"
@@ -21,10 +22,10 @@ TEST_CASE("UObject.NewObject")
 	{
 		UMyAutoRTFMTestObject* Object = nullptr;
 
-		AutoRTFM::Commit([&]
-			{
-				Object = NewObject<UMyAutoRTFMTestObject>();
-			});
+		AutoRTFM::Testing::Commit([&]
+		{
+			Object = NewObject<UMyAutoRTFMTestObject>();
+		});
 
 		REQUIRE(nullptr != Object);
 		REQUIRE(42 == Object->Value);
@@ -34,11 +35,11 @@ TEST_CASE("UObject.NewObject")
 	{
 		UMyAutoRTFMTestObject* Object = nullptr;
 
-		REQUIRE(AutoRTFM::ETransactionResult::AbortedByRequest == AutoRTFM::Transact([&]
-			{
-				Object = NewObject<UMyAutoRTFMTestObject>();
-				AutoRTFM::AbortTransaction();
-			}));
+		AutoRTFM::Testing::Abort([&]
+		{
+			Object = NewObject<UMyAutoRTFMTestObject>();
+			AutoRTFM::AbortTransaction();
+		});
 
 		REQUIRE(nullptr == Object);
 	}
@@ -51,10 +52,10 @@ TEST_CASE("UObject.NewObjectWithOuter")
 		UMyAutoRTFMTestObject* Outer = NewObject<UMyAutoRTFMTestObject>();
 		UMyAutoRTFMTestObject* Object = nullptr;
 
-		AutoRTFM::Commit([&]
-			{
-				Object = NewObject<UMyAutoRTFMTestObject>(Outer);
-			});
+		AutoRTFM::Testing::Commit([&]
+		{
+			Object = NewObject<UMyAutoRTFMTestObject>(Outer);
+		});
 
 		REQUIRE(nullptr != Object);
 		REQUIRE(42 == Object->Value);
@@ -67,17 +68,55 @@ TEST_CASE("UObject.NewObjectWithOuter")
 		UMyAutoRTFMTestObject* Outer = NewObject<UMyAutoRTFMTestObject>();
 		UMyAutoRTFMTestObject* Object = nullptr;
 
-		REQUIRE(AutoRTFM::ETransactionResult::AbortedByRequest == AutoRTFM::Transact([&]
-			{
-				Object = NewObject<UMyAutoRTFMTestObject>(Outer);
-				AutoRTFM::AbortTransaction();
-			}));
+		AutoRTFM::Testing::Abort([&]
+		{
+			Object = NewObject<UMyAutoRTFMTestObject>(Outer);
+			AutoRTFM::AbortTransaction();
+		});
 
 		REQUIRE(nullptr == Object);
 		REQUIRE(42 == Outer->Value);
 	}
 }
 
+TEST_CASE("UObject.Rename")
+{
+	SECTION("Rename")
+	{
+		UMyAutoRTFMTestObject* OuterA = NewObject<UMyAutoRTFMTestObject>();
+		UMyAutoRTFMTestObject* OuterB = NewObject<UMyAutoRTFMTestObject>();
+		UMyAutoRTFMTestObject* Object = NewObject<UMyAutoRTFMTestObject>(OuterA);
+		REQUIRE(Object->IsInOuter(OuterA));
+
+		AutoRTFM::Testing::Commit([&]
+		{
+			REQUIRE(Object->Rename(nullptr, OuterB));
+		});
+
+		REQUIRE(nullptr != Object);
+		REQUIRE(!Object->IsInOuter(OuterA));
+		REQUIRE(Object->IsInOuter(OuterB));
+	}
+
+	SECTION("Abort")
+	{
+		UMyAutoRTFMTestObject* OuterA = NewObject<UMyAutoRTFMTestObject>();
+		UMyAutoRTFMTestObject* OuterB = NewObject<UMyAutoRTFMTestObject>();
+		UMyAutoRTFMTestObject* Object = NewObject<UMyAutoRTFMTestObject>(OuterA);
+		REQUIRE(Object->IsInOuter(OuterA));
+
+		AutoRTFM::Testing::Abort([&]
+		{
+			REQUIRE(Object->Rename(nullptr, OuterB));
+			AutoRTFM::AbortTransaction();
+		});
+
+		REQUIRE(nullptr != Object);
+		REQUIRE(Object->IsInOuter(OuterA));
+		REQUIRE(!Object->IsInOuter(OuterB));
+	}
+}
+
 // This is a copy of the helper function in TestGarbageCollector.cpp.
 int32 PerformGarbageCollectionWithIncrementalReachabilityAnalysis(TFunctionRef<bool(int32)> ReachabilityIterationCallback)
 {
@@ -136,10 +175,10 @@ TEST_CASE("UObject.MarkAsReachable")
 			return true;
 		}
 
-		REQUIRE(AutoRTFM::ETransactionResult::Committed == AutoRTFM::Transact([&]
+		AutoRTFM::Testing::Commit([&]
 		{
 			Object->MarkAsReachable();
-		}));
+		});
 
 		return false;
 	});
@@ -184,7 +223,7 @@ TEST_CASE("FUObjectAnnotationSparse.AddAnnotation")
 	{
 		SECTION("Commit")
 		{
-			AutoRTFM::Transact([&]
+			AutoRTFM::Testing::Commit([&]
 			{
 				REQUIRE(FTestAnnotation() == AnnotationMap.GetAnnotation(Key));
 				AnnotationMap.AddAnnotation(Key, ValueA);
@@ -196,7 +235,7 @@ TEST_CASE("FUObjectAnnotationSparse.AddAnnotation")
 
 		SECTION("Abort")
 		{
-			AutoRTFM::Transact([&]
+			AutoRTFM::Testing::Abort([&]
 			{
 				AnnotationMap.AddAnnotation(Key, ValueA);
 				AutoRTFM::AbortTransaction();
@@ -212,7 +251,7 @@ TEST_CASE("FUObjectAnnotationSparse.AddAnnotation")
 
 		SECTION("Commit")
 		{
-			AutoRTFM::Transact([&]
+			AutoRTFM::Testing::Commit([&]
 			{
 				REQUIRE(ValueB == AnnotationMap.GetAnnotation(Key));
 				AnnotationMap.AddAnnotation(Key, ValueA);
@@ -224,7 +263,7 @@ TEST_CASE("FUObjectAnnotationSparse.AddAnnotation")
 
 		SECTION("Abort")
 		{
-			AutoRTFM::Transact([&]
+			AutoRTFM::Testing::Abort([&]
 			{
 				AnnotationMap.AddAnnotation(Key, ValueA);
 				AutoRTFM::AbortTransaction();
@@ -254,7 +293,7 @@ TEST_CASE("UObject.AnnotationMap")
 
 	UObject* Key = NewObject<UMyAutoRTFMTestObject>();
 
-	AutoRTFM::Transact([&]
+	AutoRTFM::Testing::Commit([&]
 	{
 		UObject* Value = NewObject<UMyAutoRTFMTestObject>();
 		AnnotationMap.GetAnnotation(Key);
@@ -278,21 +317,19 @@ TEST_CASE("UObject.AtomicallySetFlags")
 
 	REQUIRE(Object->HasAllFlags(OldFlags) & !Object->HasAllFlags(FlagsToAdd));
 
-	AutoRTFM::ETransactionResult Result = AutoRTFM::Transact([&]
-		{
-			Object->AtomicallySetFlags(FlagsToAdd);
-			AutoRTFM::AbortTransaction();
-		});
+	AutoRTFM::Testing::Abort([&]
+	{
+		Object->AtomicallySetFlags(FlagsToAdd);
+		AutoRTFM::AbortTransaction();
+	});
 
-	REQUIRE(AutoRTFM::ETransactionResult::AbortedByRequest == Result);
 	REQUIRE(Object->HasAllFlags(OldFlags) & !Object->HasAllFlags(FlagsToAdd));
 
-	Result = AutoRTFM::Transact([&]
-		{
-			Object->AtomicallySetFlags(FlagsToAdd);
-		});
+	AutoRTFM::Testing::Commit([&]
+	{
+		Object->AtomicallySetFlags(FlagsToAdd);
+	});
 
-	REQUIRE(AutoRTFM::ETransactionResult::Committed == Result);
 	REQUIRE(Object->HasAllFlags(OldFlags) & Object->HasAllFlags(FlagsToAdd));
 }
 
@@ -310,21 +347,19 @@ TEST_CASE("UObject.AtomicallyClearFlags")
 
 	REQUIRE(Object->HasAllFlags(OldFlags) & !Object->HasAllFlags(FlagsToClear));
 
-	AutoRTFM::ETransactionResult Result = AutoRTFM::Transact([&]
-		{
-			Object->AtomicallyClearFlags(FlagsToClear);
-			AutoRTFM::AbortTransaction();
-		});
+	AutoRTFM::Testing::Abort([&]
+	{
+		Object->AtomicallyClearFlags(FlagsToClear);
+		AutoRTFM::AbortTransaction();
+	});
 
-	REQUIRE(AutoRTFM::ETransactionResult::AbortedByRequest == Result);
 	REQUIRE(Object->HasAllFlags(OldFlags) & !Object->HasAllFlags(FlagsToClear));
 
-	Result = AutoRTFM::Transact([&]
-		{
-			Object->AtomicallyClearFlags(FlagsToClear);
-		});
+	AutoRTFM::Testing::Commit([&]
+	{
+		Object->AtomicallyClearFlags(FlagsToClear);
+	});
 
-	REQUIRE(AutoRTFM::ETransactionResult::Committed == Result);
 	REQUIRE(Object->HasAnyFlags(OldFlags) & !Object->HasAllFlags(FlagsToClear));
 }
 
@@ -355,7 +390,7 @@ TEST_CASE("UObject.FUObjectThreadContext")
 		
 		static void CtorCreateInnerTransact(const FObjectInitializer& ObjectInitializer, UMyAutoRTFMTestObject& Object)
 		{
-			AutoRTFM::Transact([&]
+			AutoRTFM::Testing::Commit([&]
 			{
 				REQUIRE(1 == FUObjectThreadContext::Get().IsInConstructor);
 				CreateObjectWithCtor(nullptr);
@@ -377,145 +412,137 @@ TEST_CASE("UObject.FUObjectThreadContext")
 
 	SECTION("Transact(UObjectCtor(Abort))")
 	{
-		AutoRTFM::ETransactionResult Result = AutoRTFM::Transact([&]
-			{
-				REQUIRE(AutoRTFM::IsClosed());
-				REQUIRE(0 == FUObjectThreadContext::Get().IsInConstructor);
-				
-				Fns::CreateObjectWithCtor(Fns::CtorAbort);
-				FAIL(/* unreachable */);
-			});
+		AutoRTFM::Testing::Abort([&]
+		{
+			REQUIRE(AutoRTFM::IsClosed());
+			REQUIRE(0 == FUObjectThreadContext::Get().IsInConstructor);
+
+			Fns::CreateObjectWithCtor(Fns::CtorAbort);
+			FAIL(/* unreachable */);
+		});
 
-		REQUIRE(AutoRTFM::ETransactionResult::AbortedByRequest == Result);
 		REQUIRE(0 == FUObjectThreadContext::Get().IsInConstructor);
 	}
-
 	SECTION("Transact(UObjectCtor(), Abort)")
 	{
-		AutoRTFM::ETransactionResult Result = AutoRTFM::Transact([&]
-			{
-				REQUIRE(AutoRTFM::IsClosed());
-				REQUIRE(0 == FUObjectThreadContext::Get().IsInConstructor);
+		AutoRTFM::Testing::Abort([&]
+		{
+			REQUIRE(AutoRTFM::IsClosed());
+			REQUIRE(0 == FUObjectThreadContext::Get().IsInConstructor);
 				
-				NewObject<UMyAutoRTFMTestObject>();
+			NewObject<UMyAutoRTFMTestObject>();
 
-				REQUIRE(0 == FUObjectThreadContext::Get().IsInConstructor);
+			REQUIRE(0 == FUObjectThreadContext::Get().IsInConstructor);
 
-				AutoRTFM::AbortTransaction();
-				FAIL(/* unreachable */);
-			});
+			AutoRTFM::AbortTransaction();
+			FAIL(/* unreachable */);
+		});
 
-		REQUIRE(AutoRTFM::ETransactionResult::AbortedByRequest == Result);
 		REQUIRE(0 == FUObjectThreadContext::Get().IsInConstructor);
 	}
 
 	SECTION("Transact(Open(UObjectCtor(Abort)), UObjectCtor(Abort))")
 	{
-		AutoRTFM::ETransactionResult Result = AutoRTFM::Transact([&]
-			{
-				REQUIRE(AutoRTFM::IsClosed());
-				REQUIRE(0 == FUObjectThreadContext::Get().IsInConstructor);
+		AutoRTFM::Testing::Abort([&]
+		{
+			REQUIRE(AutoRTFM::IsClosed());
+			REQUIRE(0 == FUObjectThreadContext::Get().IsInConstructor);
 				
-				AutoRTFM::Open([&]
-				{
-					AutoRTFM::AbortTransaction();
-				});
+			AutoRTFM::Open([&]
+			{
+				AutoRTFM::AbortTransaction();
+			});
 
-				REQUIRE(0 == FUObjectThreadContext::Get().IsInConstructor);
+			REQUIRE(0 == FUObjectThreadContext::Get().IsInConstructor);
 
-				Fns::CreateObjectWithCtor(Fns::CtorAbort);
-				FAIL(/* unreachable */);
-			});
+			Fns::CreateObjectWithCtor(Fns::CtorAbort);
+			FAIL(/* unreachable */);
+		});
 
-		REQUIRE(AutoRTFM::ETransactionResult::AbortedByRequest == Result);
 		REQUIRE(0 == FUObjectThreadContext::Get().IsInConstructor);
 	}
 
 	SECTION("Transact(Open(UObjectCtor(Transact(UObjectCtor))), UObjectCtor(Abort))")
 	{
-		AutoRTFM::ETransactionResult Result = AutoRTFM::Transact([&]
-			{
-				REQUIRE(AutoRTFM::IsClosed());
-				REQUIRE(0 == FUObjectThreadContext::Get().IsInConstructor);
+		AutoRTFM::Testing::Abort([&]
+		{
+			REQUIRE(AutoRTFM::IsClosed());
+			REQUIRE(0 == FUObjectThreadContext::Get().IsInConstructor);
 				
-				AutoRTFM::Open([&]
-				{
-					Fns::CreateObjectWithCtor(Fns::CtorCreateInnerTransact);
-				});
+			AutoRTFM::Open([&]
+			{
+				Fns::CreateObjectWithCtor(Fns::CtorCreateInnerTransact);
+			});
 
-				REQUIRE(0 == FUObjectThreadContext::Get().IsInConstructor);
+			REQUIRE(0 == FUObjectThreadContext::Get().IsInConstructor);
 
-				Fns::CreateObjectWithCtor(Fns::CtorAbort);
-				FAIL(/* unreachable */);
-			});
+			Fns::CreateObjectWithCtor(Fns::CtorAbort);
+			FAIL(/* unreachable */);
+		});
 
-		REQUIRE(AutoRTFM::ETransactionResult::AbortedByRequest == Result);
 		REQUIRE(0 == FUObjectThreadContext::Get().IsInConstructor);
 	}
 
 	SECTION("Transact(Open(UObjectCtor(Transact(UObjectCtor))), Abort)")
 	{
-		AutoRTFM::ETransactionResult Result = AutoRTFM::Transact([&]
-			{
-				REQUIRE(AutoRTFM::IsClosed());
-				REQUIRE(0 == FUObjectThreadContext::Get().IsInConstructor);
+		AutoRTFM::Testing::Abort([&]
+		{
+			REQUIRE(AutoRTFM::IsClosed());
+			REQUIRE(0 == FUObjectThreadContext::Get().IsInConstructor);
 				
-				AutoRTFM::Open([&]
-				{
-					Fns::CreateObjectWithCtor(Fns::CtorCreateInnerTransact);
-				});
+			AutoRTFM::Open([&]
+			{
+				Fns::CreateObjectWithCtor(Fns::CtorCreateInnerTransact);
+			});
 
-				REQUIRE(0 == FUObjectThreadContext::Get().IsInConstructor);
+			REQUIRE(0 == FUObjectThreadContext::Get().IsInConstructor);
 
-				AutoRTFM::AbortTransaction();
-				FAIL(/* unreachable */);
-			});
+			AutoRTFM::AbortTransaction();
+			FAIL(/* unreachable */);
+		});
 
-		REQUIRE(AutoRTFM::ETransactionResult::AbortedByRequest == Result);
 		REQUIRE(0 == FUObjectThreadContext::Get().IsInConstructor);
 	}
 
 	SECTION("Transact(Open(UObjectCtor(Close(UObjectCtor))), UObjectCtor(Abort))")
 	{
-		AutoRTFM::ETransactionResult Result = AutoRTFM::Transact([&]
-			{
-				REQUIRE(AutoRTFM::IsClosed());
-				REQUIRE(0 == FUObjectThreadContext::Get().IsInConstructor);
+		AutoRTFM::Testing::Abort([&]
+		{
+			REQUIRE(AutoRTFM::IsClosed());
+			REQUIRE(0 == FUObjectThreadContext::Get().IsInConstructor);
 				
-				AutoRTFM::Open([&]
-				{
-					Fns::CreateObjectWithCtor(Fns::CtorCreateInnerClosed);
-				});
+			AutoRTFM::Open([&]
+			{
+				Fns::CreateObjectWithCtor(Fns::CtorCreateInnerClosed);
+			});
 
-				REQUIRE(0 == FUObjectThreadContext::Get().IsInConstructor);
+			REQUIRE(0 == FUObjectThreadContext::Get().IsInConstructor);
 
-				Fns::CreateObjectWithCtor(Fns::CtorAbort);
-				FAIL(/* unreachable */);
-			});
+			Fns::CreateObjectWithCtor(Fns::CtorAbort);
+			FAIL(/* unreachable */);
+		});
 
-		REQUIRE(AutoRTFM::ETransactionResult::AbortedByRequest == Result);
 		REQUIRE(0 == FUObjectThreadContext::Get().IsInConstructor);
 	}
 	
 	SECTION("Transact(Open(UObjectCtor(Close(UObjectCtor))), Abort)")
 	{
-		AutoRTFM::ETransactionResult Result = AutoRTFM::Transact([&]
-			{
-				REQUIRE(AutoRTFM::IsClosed());
-				REQUIRE(0 == FUObjectThreadContext::Get().IsInConstructor);
+		AutoRTFM::Testing::Abort([&]
+		{
+			REQUIRE(AutoRTFM::IsClosed());
+			REQUIRE(0 == FUObjectThreadContext::Get().IsInConstructor);
 				
-				AutoRTFM::Open([&]
-				{
-					Fns::CreateObjectWithCtor(Fns::CtorCreateInnerClosed);
-				});
+			AutoRTFM::Open([&]
+			{
+				Fns::CreateObjectWithCtor(Fns::CtorCreateInnerClosed);
+			});
 
-				REQUIRE(0 == FUObjectThreadContext::Get().IsInConstructor);
+			REQUIRE(0 == FUObjectThreadContext::Get().IsInConstructor);
 
-				AutoRTFM::AbortTransaction();
-				FAIL(/* unreachable */);
-			});
+			AutoRTFM::AbortTransaction();
+			FAIL(/* unreachable */);
+		});
 
-		REQUIRE(AutoRTFM::ETransactionResult::AbortedByRequest == Result);
 		REQUIRE(0 == FUObjectThreadContext::Get().IsInConstructor);
 	}
 }
diff --git a/Engine/Source/Runtime/CoreUObject/Private/UObject/Obj.cpp b/Engine/Source/Runtime/CoreUObject/Private/UObject/Obj.cpp
--- a/Engine/Source/Runtime/CoreUObject/Private/UObject/Obj.cpp
+++ b/Engine/Source/Runtime/CoreUObject/Private/UObject/Obj.cpp
@@ -318,33 +318,32 @@ bool UObject::Rename(const TCHAR* InName, UObject* NewOuter, ERenameFlags Flags)
 	bool bCreateRedirector = false;
 	UObject* OldOuter = nullptr;
 
-	AutoRTFM::Open([&]()
+	if (!(Flags & REN_NonTransactional))
 	{
-		if (!(Flags & REN_NonTransactional))
+		// Mark touched packages as dirty.
+		if (Flags & REN_DoNotDirty)
 		{
-			// Mark touched packages as dirty.
-			if (Flags & REN_DoNotDirty)
-			{
-				// This will only mark dirty if in a transaction,
-				// the object is transactional, and the object is
-				// not in a PlayInEditor package.
-				Modify(false);
-			}
-			else
-			{
-				// This will maintain previous behavior...
-				// Which was to directly call MarkPackageDirty
-				Modify(true);
-			}
+			// This will only mark dirty if in a transaction,
+			// the object is transactional, and the object is
+			// not in a PlayInEditor package.
+			Modify(false);
 		}
-		if (NewOuter)
+		else
 		{
-			if (!(Flags & REN_DoNotDirty))
-			{
-				NewOuter->MarkPackageDirty();
-			}
+			// This will maintain previous behavior...
+			// Which was to directly call MarkPackageDirty
+			Modify(true);
 		}
+	}
+	if (NewOuter)
+	{
+		if (!(Flags & REN_DoNotDirty))
+		{
+			NewOuter->MarkPackageDirty();
+		}
+	}
 
+	{
 		// Ensure that between StaticFindObjectFast and completion of LowLevelRename nothing else modifies
 		// the UObject hash tables.
 		FScopedUObjectHashTablesLock HashTablesLock;
@@ -386,7 +385,7 @@ bool UObject::Rename(const TCHAR* InName, UObject* NewOuter, ERenameFlags Flags)
 		UE::CoreUObject::Private::UpdateRenamedObject(this, NewName, NewOuter);
 #endif
 		LowLevelRename(NewName, NewOuter);
-	});
+	}
 
 	// Create the redirector AFTER renaming the object. Two objects of different classes may not have the same fully qualified name.
 	if (bCreateRedirector)
