[AutoRTFM] Add TTask, a compact TFunction replacement.

TTask is a small std::function like class that can wrap a functor (lambda, etc).
TTask will be used by AutoRTFM for callbacks such as OnCommit and OnAbort handlers, which there may be many accumulating throughout a transaction.
For this reason TTask is designed to be compact in size, but large enough to fit the most common functors used by AutoRTFM tasks.

diff --git a/Engine/Source/Programs/AutoRTFMTests/Private/TaskTests.cpp b/Engine/Source/Programs/AutoRTFMTests/Private/TaskTests.cpp
new file mode 100644
--- /dev/null
+++ b/Engine/Source/Programs/AutoRTFMTests/Private/TaskTests.cpp
@@ -0,0 +1,425 @@
+// Copyright Epic Games, Inc. All Rights Reserved.
+
+#include "AutoRTFM/AutoRTFMTask.h"
+#include "Catch2Includes.h"
+#include "Templates/UnrealTemplate.h"
+
+TEST_CASE("TTask")
+{
+	SECTION("void()")
+	{
+		using FTask = AutoRTFM::TTask<void()>;
+
+		int ValueA = 0;
+		int ValueB = 0;
+		int ValueC = 0;
+		int ValueD = 0;
+
+		auto SmallLambda = [&ValueA]
+		{
+			ValueA++;
+		};
+
+		auto LargeLambda = [&ValueA, &ValueB, &ValueC, &ValueD]
+		{
+			ValueA++;
+			ValueB++;
+			ValueC++;
+			ValueD++;
+		};
+
+		struct FAlignedStruct
+		{
+			alignas(16) int& Value;
+		};
+		
+		auto AlignedLambda = [Struct = FAlignedStruct(ValueB)]
+		{
+			REQUIRE((reinterpret_cast<uintptr_t>(&Struct) & 15) == 0);
+			Struct.Value++;
+		};
+
+		static_assert(sizeof(SmallLambda) <= FTask::InlineDataSize);
+		static_assert(sizeof(LargeLambda) > FTask::InlineDataSize);
+		static_assert(sizeof(AlignedLambda) <= FTask::InlineDataSize);
+		static_assert(alignof(decltype(AlignedLambda)) > FTask::InlineDataAlignment);
+
+		auto WasSmallLambdaCalled = [&]
+		{
+			return ValueA == 1 && ValueB == 0 && ValueC == 0 && ValueD == 0;
+		};
+		auto WasLargeLambdaCalled = [&]
+		{
+			return ValueA == 1 && ValueB == 1 && ValueC == 1 && ValueD == 1;
+		};
+		auto WasAlignedLambdaCalled = [&]
+		{
+			return ValueA == 0 && ValueB == 1 && ValueC == 0 && ValueD == 0;
+		};
+
+		SECTION("IsValid")
+		{
+			REQUIRE(FTask().IsSet() == false);
+			REQUIRE(FTask(SmallLambda).IsSet() == true);
+			REQUIRE(FTask(LargeLambda).IsSet() == true);
+			REQUIRE(FTask(MoveTemp(SmallLambda)).IsSet() == true);
+			REQUIRE(FTask(MoveTemp(LargeLambda)).IsSet() == true);
+		}
+
+		SECTION("Copy Construct Lambda")
+		{
+			SECTION("Small")
+			{
+				FTask Task(SmallLambda);
+				Task();
+				REQUIRE(WasSmallLambdaCalled());
+			}
+			SECTION("Large")
+			{
+				FTask Task(LargeLambda);
+				Task();
+				REQUIRE(WasLargeLambdaCalled());
+			}
+			SECTION("Aligned")
+			{
+				FTask Task(AlignedLambda);
+				Task();
+				REQUIRE(WasAlignedLambdaCalled());
+			}
+		}
+
+		SECTION("Copy Construct Task")
+		{
+			SECTION("Invalid")
+			{
+				FTask Original;
+				FTask Task(Original);
+				REQUIRE(Task.IsSet() == false);
+			}
+			SECTION("Small")
+			{
+				FTask Original(SmallLambda);
+				FTask Task(Original);
+				Original.Reset();
+				Task();
+				REQUIRE(WasSmallLambdaCalled());
+			}
+			SECTION("Large")
+			{
+				FTask Original(LargeLambda);
+				FTask Task(Original);
+				Original.Reset();
+				Task();
+				REQUIRE(WasLargeLambdaCalled());
+			}
+			SECTION("Aligned")
+			{
+				FTask Original(AlignedLambda);
+				FTask Task(Original);
+				Original.Reset();
+				Task();
+				REQUIRE(WasAlignedLambdaCalled());
+			}
+		}
+
+		SECTION("Move Construct Lambda")
+		{
+			SECTION("Small")
+			{
+				FTask Task(MoveTemp(SmallLambda));
+				Task();
+				REQUIRE(WasSmallLambdaCalled());
+			}
+			SECTION("Large")
+			{
+				FTask Task(MoveTemp(LargeLambda));
+				Task();
+				REQUIRE(WasLargeLambdaCalled());
+			}
+			SECTION("Aligned")
+			{
+				FTask Task(MoveTemp(AlignedLambda));
+				Task();
+				REQUIRE(WasAlignedLambdaCalled());
+			}
+		}
+
+		SECTION("Move Construct Task")
+		{
+			SECTION("Small")
+			{
+				FTask Original(SmallLambda);
+				FTask Task(MoveTemp(Original));
+				REQUIRE(Original.IsSet() == false);
+				Task();
+				REQUIRE(WasSmallLambdaCalled());
+			}
+			SECTION("Large")
+			{
+				FTask Original(LargeLambda);
+				FTask Task(MoveTemp(Original));
+				REQUIRE(Original.IsSet() == false);
+				Task();
+				REQUIRE(WasLargeLambdaCalled());
+			}
+			SECTION("Aligned")
+			{
+				FTask Original(AlignedLambda);
+				FTask Task(MoveTemp(Original));
+				REQUIRE(Original.IsSet() == false);
+				Task();
+				REQUIRE(WasAlignedLambdaCalled());
+			}
+		}
+
+		SECTION("Reset")
+		{
+			SECTION("Invalid")
+			{
+				FTask Task;
+				Task.Reset();
+				REQUIRE(Task.IsSet() == false);
+			}
+			SECTION("Small")
+			{
+				FTask Task(SmallLambda);
+				Task.Reset();
+				REQUIRE(Task.IsSet() == false);
+			}
+			SECTION("Large")
+			{
+				FTask Task(LargeLambda);
+				Task.Reset();
+				REQUIRE(Task.IsSet() == false);
+			}
+			SECTION("Aligned")
+			{
+				FTask Task(AlignedLambda);
+				Task.Reset();
+				REQUIRE(Task.IsSet() == false);
+			}
+		}
+	}
+	
+	SECTION("int(int, bool)")
+	{
+		using FTask = AutoRTFM::TTask<int(int, bool)>;
+
+		int ValueA = 0;
+		int ValueB = 0;
+		int ValueC = 0;
+		int ValueD = 0;
+
+		auto SmallLambda = [&ValueA](int I, bool B)
+		{
+			ValueA = B ? I : -I;
+			return I;
+		};
+
+		auto LargeLambda = [&ValueA, &ValueB, &ValueC, &ValueD](int I, bool B)
+		{
+			ValueA = B ? I : -I;
+			ValueB = B ? -I : I;
+			ValueC = B ? I : -I;
+			ValueD = B ? -I : I;
+			return I;
+		};
+
+		struct FAlignedStruct
+		{
+			alignas(16) int& Value;
+		};
+		
+		auto AlignedLambda = [Struct = FAlignedStruct(ValueB)](int I, bool B)
+		{
+			REQUIRE((reinterpret_cast<uintptr_t>(&Struct) & 15) == 0);
+			Struct.Value = B ? I : -I;
+			return I;
+		};
+
+		static_assert(sizeof(SmallLambda) <= FTask::InlineDataSize);
+		static_assert(sizeof(LargeLambda) > FTask::InlineDataSize);
+		static_assert(sizeof(AlignedLambda) <= FTask::InlineDataSize);
+		static_assert(alignof(decltype(AlignedLambda)) > FTask::InlineDataAlignment);
+
+		auto WasSmallLambdaCalled = [&](int Expected)
+		{
+			return ValueA == Expected && ValueB == 0 && ValueC == 0 && ValueD == 0;
+		};
+		auto WasLargeLambdaCalled = [&](int Expected)
+		{
+			return ValueA == Expected && ValueB == -Expected && ValueC == Expected && ValueD == -Expected;
+		};
+		auto WasAlignedLambdaCalled = [&](int Expected)
+		{
+			return ValueA == 0 && ValueB == Expected && ValueC == 0 && ValueD == 0;
+		};
+
+		SECTION("IsValid")
+		{
+			REQUIRE(FTask().IsSet() == false);
+			REQUIRE(FTask(SmallLambda).IsSet() == true);
+			REQUIRE(FTask(LargeLambda).IsSet() == true);
+			REQUIRE(FTask(MoveTemp(SmallLambda)).IsSet() == true);
+			REQUIRE(FTask(MoveTemp(LargeLambda)).IsSet() == true);
+		}
+
+		SECTION("Copy Construct Lambda")
+		{
+			SECTION("Small")
+			{
+				FTask Task(SmallLambda);
+				REQUIRE(Task(10, true) == 10);
+				REQUIRE(WasSmallLambdaCalled(10));
+				REQUIRE(Task(20, false) == 20);
+				REQUIRE(WasSmallLambdaCalled(-20));
+			}
+			SECTION("Large")
+			{
+				FTask Task(LargeLambda);
+				REQUIRE(Task(10, true) == 10);
+				REQUIRE(WasLargeLambdaCalled(10));
+				REQUIRE(Task(20, false) == 20);
+				REQUIRE(WasLargeLambdaCalled(-20));
+			}
+			SECTION("Aligned")
+			{
+				FTask Task(AlignedLambda);
+				REQUIRE(Task(10, true) == 10);
+				REQUIRE(WasAlignedLambdaCalled(10));
+				REQUIRE(Task(20, false) == 20);
+				REQUIRE(WasAlignedLambdaCalled(-20));
+			}
+		}
+
+		SECTION("Copy Construct Task")
+		{
+			SECTION("Invalid")
+			{
+				FTask Original;
+				FTask Task(Original);
+				REQUIRE(Task.IsSet() == false);
+			}
+			SECTION("Small")
+			{
+				FTask Original(SmallLambda);
+				FTask Task(Original);
+				Original.Reset();
+				REQUIRE(Task(10, true) == 10);
+				REQUIRE(WasSmallLambdaCalled(10));
+				REQUIRE(Task(20, false) == 20);
+				REQUIRE(WasSmallLambdaCalled(-20));
+			}
+			SECTION("Large")
+			{
+				FTask Original(LargeLambda);
+				FTask Task(Original);
+				Original.Reset();
+				REQUIRE(Task(10, true) == 10);
+				REQUIRE(WasLargeLambdaCalled(10));
+				REQUIRE(Task(20, false) == 20);
+				REQUIRE(WasLargeLambdaCalled(-20));
+			}
+			SECTION("Aligned")
+			{
+				FTask Original(AlignedLambda);
+				FTask Task(Original);
+				Original.Reset();
+				REQUIRE(Task(10, true) == 10);
+				REQUIRE(WasAlignedLambdaCalled(10));
+				REQUIRE(Task(20, false) == 20);
+				REQUIRE(WasAlignedLambdaCalled(-20));
+			}
+		}
+
+		SECTION("Move Construct Lambda")
+		{
+			SECTION("Small")
+			{
+				FTask Task(MoveTemp(SmallLambda));
+				REQUIRE(Task(10, true) == 10);
+				REQUIRE(WasSmallLambdaCalled(10));
+				REQUIRE(Task(20, false) == 20);
+				REQUIRE(WasSmallLambdaCalled(-20));
+			}
+			SECTION("Large")
+			{
+				FTask Task(MoveTemp(LargeLambda));
+				REQUIRE(Task(10, true) == 10);
+				REQUIRE(WasLargeLambdaCalled(10));
+				REQUIRE(Task(20, false) == 20);
+				REQUIRE(WasLargeLambdaCalled(-20));
+			}
+			SECTION("Aligned")
+			{
+				FTask Task(MoveTemp(AlignedLambda));
+				REQUIRE(Task(10, true) == 10);
+				REQUIRE(WasAlignedLambdaCalled(10));
+				REQUIRE(Task(20, false) == 20);
+				REQUIRE(WasAlignedLambdaCalled(-20));
+			}
+		}
+
+		SECTION("Move Construct Task")
+		{
+			SECTION("Small")
+			{
+				FTask Original(SmallLambda);
+				FTask Task(MoveTemp(Original));
+				REQUIRE(Original.IsSet() == false);
+				REQUIRE(Task(10, true) == 10);
+				REQUIRE(WasSmallLambdaCalled(10));
+				REQUIRE(Task(20, false) == 20);
+				REQUIRE(WasSmallLambdaCalled(-20));
+			}
+			SECTION("Large")
+			{
+				FTask Original(LargeLambda);
+				FTask Task(MoveTemp(Original));
+				REQUIRE(Original.IsSet() == false);
+				REQUIRE(Task(10, true) == 10);
+				REQUIRE(WasLargeLambdaCalled(10));
+				REQUIRE(Task(20, false) == 20);
+				REQUIRE(WasLargeLambdaCalled(-20));
+			}
+			SECTION("Aligned")
+			{
+				FTask Original(AlignedLambda);
+				FTask Task(MoveTemp(Original));
+				REQUIRE(Original.IsSet() == false);
+				REQUIRE(Task(10, true) == 10);
+				REQUIRE(WasAlignedLambdaCalled(10));
+				REQUIRE(Task(20, false) == 20);
+				REQUIRE(WasAlignedLambdaCalled(-20));
+			}
+		}
+
+		SECTION("Reset")
+		{
+			SECTION("Invalid")
+			{
+				FTask Task;
+				Task.Reset();
+				REQUIRE(Task.IsSet() == false);
+			}
+			SECTION("Small")
+			{
+				FTask Task(SmallLambda);
+				Task.Reset();
+				REQUIRE(Task.IsSet() == false);
+			}
+			SECTION("Large")
+			{
+				FTask Task(LargeLambda);
+				Task.Reset();
+				REQUIRE(Task.IsSet() == false);
+			}
+			SECTION("Aligned")
+			{
+				FTask Task(AlignedLambda);
+				Task.Reset();
+				REQUIRE(Task.IsSet() == false);
+			}
+		}
+	}
+}
diff --git a/Engine/Source/Runtime/Core/Public/AutoRTFM/AutoRTFMTask.h b/Engine/Source/Runtime/Core/Public/AutoRTFM/AutoRTFMTask.h
new file mode 100644
--- /dev/null
+++ b/Engine/Source/Runtime/Core/Public/AutoRTFM/AutoRTFMTask.h
@@ -0,0 +1,233 @@
+// Copyright Epic Games, Inc. All Rights Reserved.
+
+#pragma once
+
+#include "Templates/UnrealTemplate.h"
+#include <cstddef>
+#include <type_traits>
+
+namespace AutoRTFM
+{
+
+// TTask is a small std::function like class that can wrap a functor (lambda, etc).
+// TTask is used by AutoRTFM for callbacks such as OnCommit and OnAbort
+// handlers, which there may be many accumulating throughout a transaction.
+// For this reason TTask is designed to be compact in size (smaller than 
+// TFunction), but large enough to fit the most common functors used by AutoRTFM
+// tasks.
+template<typename Signature>
+class TTask;
+
+template<typename ReturnType, typename... ParameterTypes>
+class TTask<ReturnType(ParameterTypes...)>
+{
+public:
+	static constexpr size_t InlineDataSize = 16;
+	static constexpr size_t InlineDataAlignment = 8;
+
+private:
+	// The payload holds either the inline function data, or a pointer to the
+	// data in a heap allocation.
+	union FPayload
+	{
+		std::byte Inline[InlineDataSize];
+		void* External;
+	};
+
+	// Holds the function pointers for operating on the payload.
+	struct FPayloadMethods
+	{
+		using CallFn = ReturnType(FPayload& Payload, ParameterTypes...);
+		using CopyFn = void(FPayload& Dst, const FPayload& Src);
+		using MoveFn = void(FPayload& Dst, FPayload& Src);
+		using DestructFn = void(FPayload& Payload);
+
+		CallFn* Call = nullptr; // Call the function held by the payload.
+		CopyFn* Copy = nullptr; // Copy the payload.
+		MoveFn* Move = nullptr; // Move the payload.
+		DestructFn* Destruct = nullptr; // Destruct the payload.
+	};
+	
+	// Traits helper for a functor (T).
+	template<typename T>
+	struct TTraits
+	{
+		// True iff T can fit within the inline payload, and does not require
+		// courser alignment than provided by the inline payload.
+		static constexpr bool bCanUseInline = sizeof(T) <= InlineDataSize && alignof(T) <= InlineDataAlignment;
+
+		// Returns the functor from the payload.
+		static T* FunctionFrom(FPayload& Payload)
+		{
+			return bCanUseInline ? reinterpret_cast<T*>(Payload.Inline) : reinterpret_cast<T*>(Payload.External);
+		}
+
+		// Returns the functor from the payload.
+		static const T* FunctionFrom(const FPayload& Payload)
+		{
+			return bCanUseInline ? reinterpret_cast<const T*>(Payload.Inline) : reinterpret_cast<const T*>(Payload.External);
+		}
+
+		// The methods to operate on a payload for a functor of type T.
+		static constexpr FPayloadMethods PayloadMethods = 
+			bCanUseInline
+				? FPayloadMethods
+				{
+					/* Call */ [](FPayload& Payload, ParameterTypes... Arguments) -> ReturnType
+					{
+						return (*FunctionFrom(Payload))(Forward<ParameterTypes>(Arguments)...);
+					},
+					/* Copy */ [](FPayload& Dst, const FPayload& Src)
+					{
+						new(Dst.Inline) T(*FunctionFrom(Src));
+					},
+					/* Move */ [](FPayload& Dst, FPayload& Src)
+					{
+						new(Dst.Inline) T(MoveTemp(*FunctionFrom(Src)));
+					},
+					/* Destruct */ [](FPayload& Payload)
+					{
+						FunctionFrom(Payload)->~T();
+					},
+				}
+				: FPayloadMethods
+				{
+					/* Call */ [](FPayload& Payload, ParameterTypes... Arguments) -> ReturnType
+					{
+						return (*FunctionFrom(Payload))(Forward<ParameterTypes>(Arguments)...);
+					},
+					/* Copy */ [](FPayload& Dst, const FPayload& Src)
+					{
+						Dst.External = new T(*FunctionFrom(Src));
+					},
+					/* Move */ [](FPayload& Dst, FPayload& Src)
+					{
+						Dst.External = Src.External;
+					},
+					/* Destruct */ [](FPayload& Payload)
+					{
+						delete FunctionFrom(Payload);
+					},
+				};
+	};
+public:
+	// Constructor.
+	// The TTask is constructed in a non-set state.
+	TTask() = default;
+
+	// Destructor.
+	~TTask()
+	{
+		Reset();
+	}
+
+	// Copy constructor
+	TTask(const TTask& Other)
+	{
+		if (Other.IsSet())
+		{
+			PayloadMethods = Other.PayloadMethods;
+			PayloadMethods->Copy(Payload, Other.Payload);
+		}
+	}
+
+	// Move constructor
+	TTask(TTask&& Other)
+	{
+		if (Other.IsSet())
+		{
+			PayloadMethods = Other.PayloadMethods;
+			PayloadMethods->Move(Payload, Other.Payload);
+			Other.PayloadMethods = nullptr;
+		}
+	}
+
+	// Constructor from functor.
+	template
+	<
+		typename FunctorType
+		UE_REQUIRES(
+			// Use the explicit TTask copy / move constructors instead of this constructor
+			!std::is_same_v<std::decay_t<FunctorType>, TTask> &&
+			// The templated argument must be a functor with the correct signature
+			std::is_invocable_r_v<ReturnType, std::decay_t<FunctorType>, ParameterTypes...>)
+	>
+	TTask(FunctorType&& Other)
+	{
+		using Decayed = std::decay_t<FunctorType>;
+		using Traits = TTraits<Decayed>;
+		PayloadMethods = &Traits::PayloadMethods;
+		if (Traits::bCanUseInline)
+		{
+			new(Payload.Inline) Decayed(Forward<FunctorType>(Other));
+		}
+		else
+		{
+			Payload.External = new Decayed(Forward<FunctorType>(Other));
+		}
+	}
+
+	// Copy assignment operator
+	TTask& operator = (const TTask& Other)
+	{
+		if (&Other != this)
+		{
+			Reset();
+			if (Other.IsSet())
+			{
+				PayloadMethods = Other.PayloadMethods;
+				PayloadMethods->Copy(Payload, Other.Payload);
+			}
+		}
+		return *this;
+	}
+
+	// Move assignment operator
+	TTask& operator = (TTask&& Other)
+	{
+		if (&Other != this)
+		{
+			Reset();
+			if (Other.IsSet())
+			{
+				PayloadMethods = Other.PayloadMethods;
+				PayloadMethods->Move(Payload, Other.Payload);
+				Other.PayloadMethods = nullptr;
+			}
+		}
+		return *this;
+	}
+
+	// Call operator
+	template<typename... ArgumentTypes>
+	ReturnType operator()(ArgumentTypes&&... Arguments) const
+	{
+		return PayloadMethods->Call(Payload, Forward<ArgumentTypes>(Arguments)...);
+	}
+
+	// Reset the task to an unset state.
+	void Reset()
+	{
+		if (PayloadMethods)
+		{
+			PayloadMethods->Destruct(Payload);
+			PayloadMethods = nullptr;
+		}
+	}
+
+	// Returns true if the task holds a functor.
+	bool IsSet() const
+	{
+		return PayloadMethods != nullptr;
+	}
+
+private:
+	// mutable as the functor might be mutable.
+	alignas(InlineDataAlignment) mutable FPayload Payload;
+	FPayloadMethods const* PayloadMethods = nullptr;
+};
+
+// Tasks are designed to be compact.
+static_assert(sizeof(TTask<void()>) <= 24);
+
+} // namespace AutoRTFM
