[AutoRTFM] Flesh out TTask unit tests

Check [copy, move] x [constructor,assignment] from TFunction and std::function.

diff --git a/Engine/Source/Programs/AutoRTFMTests/Private/TaskTests.cpp b/Engine/Source/Programs/AutoRTFMTests/Private/TaskTests.cpp
--- a/Engine/Source/Programs/AutoRTFMTests/Private/TaskTests.cpp
+++ b/Engine/Source/Programs/AutoRTFMTests/Private/TaskTests.cpp
@@ -2,7 +2,9 @@
 
 #include "AutoRTFM/AutoRTFMTask.h"
 #include "Catch2Includes.h"
+#include "GenericPlatform/GenericPlatformMisc.h"
 #include "Templates/UnrealTemplate.h"
+#include <functional>
 
 TEST_CASE("TTask")
 {
@@ -66,6 +68,40 @@ TEST_CASE("TTask")
 			REQUIRE(FTask(MoveTemp(LargeLambda)).IsSet() == true);
 		}
 
+		SECTION("Copy Construct TTask")
+		{
+			SECTION("Invalid")
+			{
+				FTask Original;
+				FTask Task(Original);
+				REQUIRE(Task.IsSet() == false);
+			}
+			SECTION("Small")
+			{
+				FTask Original(SmallLambda);
+				FTask Task(Original);
+				Original.Reset();
+				Task();
+				REQUIRE(WasSmallLambdaCalled());
+			}
+			SECTION("Large")
+			{
+				FTask Original(LargeLambda);
+				FTask Task(Original);
+				Original.Reset();
+				Task();
+				REQUIRE(WasLargeLambdaCalled());
+			}
+			SECTION("Aligned")
+			{
+				FTask Original(AlignedLambda);
+				FTask Task(Original);
+				Original.Reset();
+				Task();
+				REQUIRE(WasAlignedLambdaCalled());
+			}
+		}
+
 		SECTION("Copy Construct Lambda")
 		{
 			SECTION("Small")
@@ -88,18 +124,38 @@ TEST_CASE("TTask")
 			}
 		}
 
-		SECTION("Copy Construct Task")
+		SECTION("Copy Construct TFunction")
+		{
+			TFunction<void()> Function(SmallLambda);
+			FTask Task(Function);
+			REQUIRE(Function.IsSet() == true);
+			Task();
+			REQUIRE(WasSmallLambdaCalled());
+		}
+
+		SECTION("Copy Construct std::function")
+		{
+			std::function<void()> Function(SmallLambda);
+			FTask Task(Function);
+			REQUIRE(static_cast<bool>(Function) == true);
+			Task();
+			REQUIRE(WasSmallLambdaCalled());
+		}
+
+		SECTION("Copy Assign TTask")
 		{
 			SECTION("Invalid")
 			{
 				FTask Original;
-				FTask Task(Original);
+				FTask Task;
+				Task = Original;
 				REQUIRE(Task.IsSet() == false);
 			}
 			SECTION("Small")
 			{
 				FTask Original(SmallLambda);
-				FTask Task(Original);
+				FTask Task;
+				Task = Original;
 				Original.Reset();
 				Task();
 				REQUIRE(WasSmallLambdaCalled());
@@ -107,7 +163,8 @@ TEST_CASE("TTask")
 			SECTION("Large")
 			{
 				FTask Original(LargeLambda);
-				FTask Task(Original);
+				FTask Task;
+				Task = Original;
 				Original.Reset();
 				Task();
 				REQUIRE(WasLargeLambdaCalled());
@@ -115,13 +172,59 @@ TEST_CASE("TTask")
 			SECTION("Aligned")
 			{
 				FTask Original(AlignedLambda);
-				FTask Task(Original);
+				FTask Task;
+				Task = Original;
 				Original.Reset();
 				Task();
 				REQUIRE(WasAlignedLambdaCalled());
 			}
 		}
 
+		SECTION("Copy Assign Lambda")
+		{
+			SECTION("Small")
+			{
+				FTask Task;
+				Task = SmallLambda;
+				Task();
+				REQUIRE(WasSmallLambdaCalled());
+			}
+			SECTION("Large")
+			{
+				FTask Task;
+				Task = LargeLambda;
+				Task();
+				REQUIRE(WasLargeLambdaCalled());
+			}
+			SECTION("Aligned")
+			{
+				FTask Task;
+				Task = AlignedLambda;
+				Task();
+				REQUIRE(WasAlignedLambdaCalled());
+			}
+		}
+
+		SECTION("Copy Assign TFunction")
+		{
+			TFunction<void()> Function(SmallLambda);
+			FTask Task;
+			Task = Function;
+			REQUIRE(Function.IsSet() == true);
+			Task();
+			REQUIRE(WasSmallLambdaCalled());
+		}
+
+		SECTION("Copy Assign std::function")
+		{
+			std::function<void()> Function(SmallLambda);
+			FTask Task;
+			Task = Function;
+			REQUIRE(static_cast<bool>(Function) == true);
+			Task();
+			REQUIRE(WasSmallLambdaCalled());
+		}
+
 		SECTION("Move Construct Lambda")
 		{
 			SECTION("Small")
@@ -144,7 +247,7 @@ TEST_CASE("TTask")
 			}
 		}
 
-		SECTION("Move Construct Task")
+		SECTION("Move Construct TTask")
 		{
 			SECTION("Small")
 			{
@@ -172,6 +275,100 @@ TEST_CASE("TTask")
 			}
 		}
 
+		SECTION("Move Construct TFunction")
+		{
+			TFunction<void()> Function(SmallLambda);
+			FTask Task(MoveTemp(Function));
+			REQUIRE(Function.IsSet() == false);
+			Task();
+			REQUIRE(WasSmallLambdaCalled());
+		}
+
+		SECTION("Move Construct std::function")
+		{
+			std::function<void()> Function(SmallLambda);
+			FTask Task(MoveTemp(Function));
+			REQUIRE(static_cast<bool>(Function) == false);
+			Task();
+			REQUIRE(WasSmallLambdaCalled());
+		}
+
+		SECTION("Move Assign Lambda")
+		{
+			SECTION("Small")
+			{
+				FTask Task;
+				Task = MoveTemp(SmallLambda);
+				Task();
+				REQUIRE(WasSmallLambdaCalled());
+			}
+			SECTION("Large")
+			{
+				FTask Task;
+				Task = MoveTemp(LargeLambda);
+				Task();
+				REQUIRE(WasLargeLambdaCalled());
+			}
+			SECTION("Aligned")
+			{
+				FTask Task;
+				Task = MoveTemp(AlignedLambda);
+				Task();
+				REQUIRE(WasAlignedLambdaCalled());
+			}
+		}
+
+		SECTION("Move Assign TTask")
+		{
+			SECTION("Small")
+			{
+				FTask Original(SmallLambda);
+				FTask Task;
+				Task = MoveTemp(Original);
+				REQUIRE(Original.IsSet() == false);
+				Task();
+				REQUIRE(WasSmallLambdaCalled());
+			}
+			SECTION("Large")
+			{
+				FTask Original(LargeLambda);
+				FTask Task;
+				Task = MoveTemp(Original);
+				REQUIRE(Original.IsSet() == false);
+				Task();
+				REQUIRE(WasLargeLambdaCalled());
+			}
+			SECTION("Aligned")
+			{
+				FTask Original(AlignedLambda);
+				FTask Task;
+				Task = MoveTemp(Original);
+				REQUIRE(Original.IsSet() == false);
+				Task();
+				REQUIRE(WasAlignedLambdaCalled());
+			}
+		}
+
+		SECTION("Move Assign TFunction")
+		{
+			TFunction<void()> Function(SmallLambda);
+			FTask Task;
+			Task = MoveTemp(Function);
+			REQUIRE(Function.IsSet() == false);
+			Task();
+			REQUIRE(WasSmallLambdaCalled());
+		}
+
+		SECTION("Move Assign std::function")
+		{
+			std::function<void()> Function(SmallLambda);
+			FTask Task;
+			Task = MoveTemp(Function);
+			REQUIRE(static_cast<bool>(Function) == false);
+			Task();
+			REQUIRE(WasSmallLambdaCalled());
+		}
+
 		SECTION("Reset")
 		{
 			SECTION("Invalid")
@@ -292,7 +489,7 @@ TEST_CASE("TTask")
 			}
 		}
 
-		SECTION("Copy Construct Task")
+		SECTION("Copy Construct TTask")
 		{
 			SECTION("Invalid")
 			{
@@ -360,7 +557,7 @@ TEST_CASE("TTask")
 			}
 		}
 
-		SECTION("Move Construct Task")
+		SECTION("Move Construct TTask")
 		{
 			SECTION("Small")
 			{
