[AutoRTFM] Add unit test for LoadFileToString with an IPlatformFile.

diff --git a/Engine/Source/Programs/AutoRTFMTests/Private/UECoreTests.cpp b/Engine/Source/Programs/AutoRTFMTests/Private/UECoreTests.cpp
--- a/Engine/Source/Programs/AutoRTFMTests/Private/UECoreTests.cpp
+++ b/Engine/Source/Programs/AutoRTFMTests/Private/UECoreTests.cpp
@@ -22,10 +22,12 @@
 #include "UObject/UObjectGlobals.h"
 #include "MyAutoRTFMTestObject.h"
 #include "Containers/Queue.h"
-#include "Misc/ScopeRWLock.h"
+#include "Misc/CString.h"
 #include "Misc/ConfigCacheIni.h"
 #include "Misc/CoreDelegates.h"
+#include "Misc/FileHelper.h"
 #include "Misc/PackageName.h"
+#include "Misc/ScopeRWLock.h"
 #include "StructUtils/PropertyBag.h"
 #include "Utils.h"
 #include <thread>
@@ -3444,3 +3446,113 @@ TEST_CASE("UECore.ConsoleManager")
 
 	REQUIRE(Thing == 13.0f);
 }
+
+namespace {
+
+class FakeFileHandle : public IFileHandle
+{
+public:
+	bool SeekFromEnd(int64 NewPositionRelativeToEnd = 0) override				{ unimplemented(); }
+	bool ReadAt(uint8* Destination, int64 BytesToRead, int64 Offset) override	{ unimplemented(); }
+	bool Write(const uint8* Source, int64 BytesToWrite) override				{ unimplemented(); }
+	bool Truncate(int64 NewSize) override										{ unimplemented(); }
+	bool Flush(const bool bFullFlush = false) override							{}
+	void ShrinkBuffers() override												{}
+	
+	int64 Tell() override
+	{ 
+		return Cursor; 
+	}
+
+	int64 Size() override
+	{ 
+		return strlen(Data); 
+	}
+
+	bool Seek(int64 NewPosition) override
+	{
+		check(NewPosition >= 0);
+		check(NewPosition <= Size());
+		Cursor = NewPosition;
+		return true;
+	}
+
+	bool Read(uint8* Destination, int64 BytesToRead) override
+	{
+		check(BytesToRead <= Size() - Cursor);
+		const uint8* DataPtr = reinterpret_cast<const uint8*>(Data);
+		memcpy(Destination, DataPtr + Cursor, BytesToRead);
+		return true;
+	}
+
+	const char* const Data = "File Loaded";
+	int64 Cursor = 0;
+};
+
+class FakePlatformFile : public IPlatformFile
+{
+public:
+	FakePlatformFile() {}
+
+	bool Initialize(IPlatformFile* Inner, const TCHAR* CmdLine) override                                  { unimplemented(); }
+	IPlatformFile* GetLowerLevel() override                                                               { unimplemented(); }
+	void SetLowerLevel(IPlatformFile* NewLowerLevel) override                                             { unimplemented(); }
+	bool FileExists(const TCHAR* Filename) override                                                       { unimplemented(); }
+	int64 FileSize(const TCHAR* Filename) override                                                        { unimplemented(); }
+	bool DeleteFile(const TCHAR* Filename) override                                                       { unimplemented(); }
+	bool IsReadOnly(const TCHAR* Filename) override                                                       { unimplemented(); }
+	bool MoveFile(const TCHAR* To, const TCHAR* From) override                                            { unimplemented(); }
+	bool SetReadOnly(const TCHAR* Filename, bool bNewReadOnlyValue) override                              { unimplemented(); }
+	FDateTime GetTimeStamp(const TCHAR* Filename) override                                                { unimplemented(); }
+	void SetTimeStamp(const TCHAR* Filename, FDateTime DateTime) override                                 { unimplemented(); }
+	FDateTime GetAccessTimeStamp(const TCHAR* Filename) override                                          { unimplemented(); }
+	FString GetFilenameOnDisk(const TCHAR* Filename) override                                             { unimplemented(); }
+	IFileHandle* OpenWrite(const TCHAR* Filename, bool bAppend = false, bool bAllowRead = false) override { unimplemented(); }
+	bool DirectoryExists(const TCHAR* Directory) override                                                 { unimplemented(); }
+	bool CreateDirectory(const TCHAR* Directory) override                                                 { unimplemented(); }
+	bool DeleteDirectory(const TCHAR* Directory) override                                                 { unimplemented(); }
+	FFileStatData GetStatData(const TCHAR* FilenameOrDirectory) override                                  { unimplemented(); }
+	bool IterateDirectory(const TCHAR* Directory, FDirectoryVisitor& Visitor) override                    { unimplemented(); }
+	bool IterateDirectoryStat(const TCHAR* Directory, FDirectoryStatVisitor& Visitor) override            { unimplemented(); }
+
+	const TCHAR* GetName() const override 
+	{
+		return TEXT("FakePlatformFile");
+	}
+
+	IFileHandle* OpenRead(const TCHAR* Filename, bool bAllowWrite = false) override
+	{
+		check(FString(Filename) == TEXT("FakePlatformFile"));
+		check(!bAllowWrite);
+		return new FakeFileHandle;
+	}
+};
+
+}
+
+TEST_CASE("UECore.LoadFileToString.IPlatformFile")
+{
+	FString FileData = TEXT("Nothing Happened");
+	FakePlatformFile FakeFile;
+
+	SECTION("Abort")
+	{
+		AutoRTFM::Testing::Abort([&]
+		{
+			REQUIRE(FFileHelper::LoadFileToString(FileData, &FakeFile, TEXT("FakePlatformFile")));
+			AutoRTFM::AbortTransaction();
+		});
+
+		REQUIRE(FileData == TEXT("Nothing Happened"));
+	}
+
+	SECTION("Commit")
+	{
+		AutoRTFM::Testing::Commit([&]
+		{
+			REQUIRE(FFileHelper::LoadFileToString(FileData, &FakeFile, TEXT("FakePlatformFile")));
+		});
+
+		REQUIRE(FileData == TEXT("File Loaded"));
+	}
+}
