[AutoRTFM] Fix build warnings on Horde.

This builds fine on Windows, but triggered -Wreturn-type on Linux.

```
D:\build\FNRS\Sync\Engine\Source\Programs\AutoRTFMTests\Private\UECoreTests.cpp(3455,86): error: non-void function does not return a value [-Werror,-Wreturn-type]
```

diff --git a/Engine/Source/Programs/AutoRTFMTests/Private/UECoreTests.cpp b/Engine/Source/Programs/AutoRTFMTests/Private/UECoreTests.cpp
--- a/Engine/Source/Programs/AutoRTFMTests/Private/UECoreTests.cpp
+++ b/Engine/Source/Programs/AutoRTFMTests/Private/UECoreTests.cpp
@@ -3452,11 +3452,11 @@ namespace {
 class FakeFileHandle : public IFileHandle
 {
 public:
-	bool SeekFromEnd(int64 NewPositionRelativeToEnd = 0) override				{ unimplemented(); }
-	bool ReadAt(uint8* Destination, int64 BytesToRead, int64 Offset) override	{ unimplemented(); }
-	bool Write(const uint8* Source, int64 BytesToWrite) override				{ unimplemented(); }
-	bool Truncate(int64 NewSize) override										{ unimplemented(); }
-	bool Flush(const bool bFullFlush = false) override							{}
+	bool SeekFromEnd(int64 NewPositionRelativeToEnd = 0) override				{ unimplemented(); return false; }
+	bool ReadAt(uint8* Destination, int64 BytesToRead, int64 Offset) override	{ unimplemented(); return false; }
+	bool Write(const uint8* Source, int64 BytesToWrite) override				{ unimplemented(); return false; }
+	bool Truncate(int64 NewSize) override										{ unimplemented(); return false; }
+	bool Flush(const bool bFullFlush = false) override							{ return true; }
 	void ShrinkBuffers() override												{}
 	
 	int64 Tell() override
@@ -3494,26 +3494,26 @@ class FakePlatformFile : public IPlatformFile
 public:
 	FakePlatformFile() {}
 
-	bool Initialize(IPlatformFile* Inner, const TCHAR* CmdLine) override                                  { unimplemented(); }
-	IPlatformFile* GetLowerLevel() override                                                               { unimplemented(); }
+	bool Initialize(IPlatformFile* Inner, const TCHAR* CmdLine) override                                  { unimplemented(); return false; }
+	IPlatformFile* GetLowerLevel() override                                                               { unimplemented(); return nullptr; }
 	void SetLowerLevel(IPlatformFile* NewLowerLevel) override                                             { unimplemented(); }
-	bool FileExists(const TCHAR* Filename) override                                                       { unimplemented(); }
-	int64 FileSize(const TCHAR* Filename) override                                                        { unimplemented(); }
-	bool DeleteFile(const TCHAR* Filename) override                                                       { unimplemented(); }
-	bool IsReadOnly(const TCHAR* Filename) override                                                       { unimplemented(); }
-	bool MoveFile(const TCHAR* To, const TCHAR* From) override                                            { unimplemented(); }
-	bool SetReadOnly(const TCHAR* Filename, bool bNewReadOnlyValue) override                              { unimplemented(); }
-	FDateTime GetTimeStamp(const TCHAR* Filename) override                                                { unimplemented(); }
+	bool FileExists(const TCHAR* Filename) override                                                       { unimplemented(); return false; }
+	int64 FileSize(const TCHAR* Filename) override                                                        { unimplemented(); return 0; }
+	bool DeleteFile(const TCHAR* Filename) override                                                       { unimplemented(); return false; }
+	bool IsReadOnly(const TCHAR* Filename) override                                                       { unimplemented(); return false; }
+	bool MoveFile(const TCHAR* To, const TCHAR* From) override                                            { unimplemented(); return false; }
+	bool SetReadOnly(const TCHAR* Filename, bool bNewReadOnlyValue) override                              { unimplemented(); return false; }
+	FDateTime GetTimeStamp(const TCHAR* Filename) override                                                { unimplemented(); return FDateTime{}; }
 	void SetTimeStamp(const TCHAR* Filename, FDateTime DateTime) override                                 { unimplemented(); }
-	FDateTime GetAccessTimeStamp(const TCHAR* Filename) override                                          { unimplemented(); }
-	FString GetFilenameOnDisk(const TCHAR* Filename) override                                             { unimplemented(); }
-	IFileHandle* OpenWrite(const TCHAR* Filename, bool bAppend = false, bool bAllowRead = false) override { unimplemented(); }
-	bool DirectoryExists(const TCHAR* Directory) override                                                 { unimplemented(); }
-	bool CreateDirectory(const TCHAR* Directory) override                                                 { unimplemented(); }
-	bool DeleteDirectory(const TCHAR* Directory) override                                                 { unimplemented(); }
-	FFileStatData GetStatData(const TCHAR* FilenameOrDirectory) override                                  { unimplemented(); }
-	bool IterateDirectory(const TCHAR* Directory, FDirectoryVisitor& Visitor) override                    { unimplemented(); }
-	bool IterateDirectoryStat(const TCHAR* Directory, FDirectoryStatVisitor& Visitor) override            { unimplemented(); }
+	FDateTime GetAccessTimeStamp(const TCHAR* Filename) override                                          { unimplemented(); return FDateTime{}; }
+	FString GetFilenameOnDisk(const TCHAR* Filename) override                                             { unimplemented(); return FString{}; }
+	IFileHandle* OpenWrite(const TCHAR* Filename, bool bAppend = false, bool bAllowRead = false) override { unimplemented(); return nullptr; }
+	bool DirectoryExists(const TCHAR* Directory) override                                                 { unimplemented(); return false; }
+	bool CreateDirectory(const TCHAR* Directory) override                                                 { unimplemented(); return false; }
+	bool DeleteDirectory(const TCHAR* Directory) override                                                 { unimplemented(); return false; }
+	FFileStatData GetStatData(const TCHAR* FilenameOrDirectory) override                                  { unimplemented(); return FFileStatData{}; }
+	bool IterateDirectory(const TCHAR* Directory, FDirectoryVisitor& Visitor) override                    { unimplemented(); return false; }
+	bool IterateDirectoryStat(const TCHAR* Directory, FDirectoryStatVisitor& Visitor) override            { unimplemented(); return false; }
 
 	const TCHAR* GetName() const override 
 	{
