[AutoRTFM] Fix TTask test that assumed that moving a std::function would invalidate it.

This is not guaranteed, and does not happen on the standard library used by our linux builds.
Also fix the name of a SECTION after the IsValid() method was renamed to IsSet().

diff --git a/Engine/Source/Programs/AutoRTFMTests/Private/TaskTests.cpp b/Engine/Source/Programs/AutoRTFMTests/Private/TaskTests.cpp
--- a/Engine/Source/Programs/AutoRTFMTests/Private/TaskTests.cpp
+++ b/Engine/Source/Programs/AutoRTFMTests/Private/TaskTests.cpp
@@ -59,7 +59,7 @@ TEST_CASE("TTask")
 			return ValueA == 0 && ValueB == 1 && ValueC == 0 && ValueD == 0;
 		};
 
-		SECTION("IsValid")
+		SECTION("IsSet")
 		{
 			REQUIRE(FTask().IsSet() == false);
 			REQUIRE(FTask(SmallLambda).IsSet() == true);
@@ -288,7 +288,8 @@ TEST_CASE("TTask")
 		{
 			std::function<void()> Function(SmallLambda);
 			FTask Task(MoveTemp(Function));
-			REQUIRE(static_cast<bool>(Function) == false);
+			// Note: standard library implementations of std::function are not
+			// guaranteed to invalidate the function on move - hence no check.
 			Task();
 			REQUIRE(WasSmallLambdaCalled());
 		}
@@ -364,7 +365,8 @@ TEST_CASE("TTask")
 			std::function<void()> Function(SmallLambda);
 			FTask Task;
 			Task = MoveTemp(Function);
-			REQUIRE(static_cast<bool>(Function) == false);
+			// Note: standard library implementations of std::function are not
+			// guaranteed to invalidate the function on move - hence no check.
 			Task();
 			REQUIRE(WasSmallLambdaCalled());
 		}
