[AutoRTFM] Fix renaming of UObject within a transaction

There were multiple issues at play here:
* Due to recent changes, we were having mixed open / closed writes into the UObjectHash tables, which caused corruption on abort. All calls into UObjectHash must be in the open, so I've added ensures and UE_AUTORTFM_NOAUTORTFM to functions that mutate the ObjectOuterMap. This will cause future closed calls to crash early instead of producing subtle bugs.
* The memory for a reconstructed UObject was snapshotted *after* ConditionalBeginDestroy() / ConditionalFinishDestroy(), meaning the restored object was left in a partially-destructed state. The snapshot is now done before these calls.
* The AutoRTFM custom handling for object renaming was recently removed, because it was assumed that AutoRTFM's automatic memory tracking would take care of restoring the name / outer on abort. However, LowLevelRename() calls into [Unh|H]ashObject(), so we need to do this in the open. Wrap LowLevelRename() with an explicit onabort handler to undo the rename.

diff --git a/Engine/Source/Programs/AutoRTFMTests/Private/UObjectTests.cpp b/Engine/Source/Programs/AutoRTFMTests/Private/UObjectTests.cpp
--- a/Engine/Source/Programs/AutoRTFMTests/Private/UObjectTests.cpp
+++ b/Engine/Source/Programs/AutoRTFMTests/Private/UObjectTests.cpp
@@ -81,39 +81,237 @@ TEST_CASE("UObject.NewObjectWithOuter")
 
 TEST_CASE("UObject.Rename")
 {
-	SECTION("Rename")
+	const TCHAR* Cat = TEXT("Cat");
+	const TCHAR* Dog = TEXT("Dog");
+	const TCHAR* Bat = TEXT("Bat");
+	UMyAutoRTFMTestObject* OuterA = NewObject<UMyAutoRTFMTestObject>();
+	UMyAutoRTFMTestObject* OuterB = NewObject<UMyAutoRTFMTestObject>();
+	UMyAutoRTFMTestObject* OuterC = NewObject<UMyAutoRTFMTestObject>();
+	UMyAutoRTFMTestObject* Object = NewObject<UMyAutoRTFMTestObject>(OuterA, Cat);
+	REQUIRE(Object->GetOuter() == OuterA);
+
+	SECTION("Commit(Rename(Name))")
 	{
-		UMyAutoRTFMTestObject* OuterA = NewObject<UMyAutoRTFMTestObject>();
-		UMyAutoRTFMTestObject* OuterB = NewObject<UMyAutoRTFMTestObject>();
-		UMyAutoRTFMTestObject* Object = NewObject<UMyAutoRTFMTestObject>(OuterA);
-		REQUIRE(Object->IsInOuter(OuterA));
+		AutoRTFM::Testing::Commit([&]
+		{
+			REQUIRE(Object->Rename(Dog, OuterA));
+		});
+
+		REQUIRE(Object->GetName() == Dog);
+		REQUIRE(Object->GetOuter() == OuterA);
+	}
+
+	SECTION("Abort(Rename(Name))")
+	{
+		AutoRTFM::Testing::Abort([&]
+		{
+			REQUIRE(Object->Rename(Dog, OuterA));
+			AutoRTFM::AbortTransaction();
+		});
+
+		REQUIRE(Object->GetName() == Cat);
+		REQUIRE(Object->GetOuter() == OuterA);
+	}
+
+	SECTION("Commit(Rename(Name), Rename(Name))")
+	{
+		AutoRTFM::Testing::Commit([&]
+		{
+			REQUIRE(Object->Rename(Dog, OuterA));
+			REQUIRE(Object->Rename(Bat, OuterA));
+		});
+
+		REQUIRE(Object->GetName() == Bat);
+		REQUIRE(Object->GetOuter() == OuterA);
+	}
+
+	SECTION("Abort(Rename(Name), Rename(Name))")
+	{
+		AutoRTFM::Testing::Abort([&]
+		{
+			REQUIRE(Object->Rename(Dog, OuterA));
+			REQUIRE(Object->Rename(Bat, OuterA));
+			AutoRTFM::AbortTransaction();
+		});
+
+		REQUIRE(Object->GetName() == Cat);
+		REQUIRE(Object->GetOuter() == OuterA);
+	}
+
+	SECTION("Commit(Rename(Name), Commit(Rename(Name)))")
+	{
+		AutoRTFM::Testing::Commit([&]
+		{
+			REQUIRE(Object->Rename(Dog, OuterA));
+			AutoRTFM::Testing::Commit([&]
+			{
+				REQUIRE(Object->Rename(Bat, OuterA));
+			});
+		});
+
+		REQUIRE(Object->GetName() == Bat);
+		REQUIRE(Object->GetOuter() == OuterA);
+	}
+
+	SECTION("Commit(Rename(Name), Abort(Rename(Name)))")
+	{
+		AutoRTFM::Testing::Commit([&]
+		{
+			REQUIRE(Object->Rename(Dog, OuterA));
+			AutoRTFM::Testing::Abort([&]
+			{
+				REQUIRE(Object->Rename(Bat, OuterA));
+				AutoRTFM::AbortTransaction();
+			});
+		});
+
+		REQUIRE(Object->GetName() == Dog);
+		REQUIRE(Object->GetOuter() == OuterA);
+	}
+
+	SECTION("Abort(Rename(Name), Commit(Rename(Name)))")
+	{
+		AutoRTFM::Testing::Abort([&]
+		{
+			REQUIRE(Object->Rename(Bat, OuterA));
+			AutoRTFM::Testing::Commit([&]
+			{
+				REQUIRE(Object->Rename(Dog, OuterA));
+			});
+			AutoRTFM::AbortTransaction();
+		});
+
+		REQUIRE(Object->GetName() == Cat);
+		REQUIRE(Object->GetOuter() == OuterA);
+	}
+
+	SECTION("Abort(Rename(Name), Commit(Rename(Name)))")
+	{
+		AutoRTFM::Testing::Abort([&]
+		{
+			REQUIRE(Object->Rename(Bat, OuterA));
+			AutoRTFM::Testing::Abort([&]
+			{
+				REQUIRE(Object->Rename(Dog, OuterA));
+				AutoRTFM::AbortTransaction();
+			});
+			AutoRTFM::AbortTransaction();
+		});
 
+		REQUIRE(Object->GetName() == Cat);
+		REQUIRE(Object->GetOuter() == OuterA);
+	}
+
+	SECTION("Commit(Rename(Object))")
+	{
 		AutoRTFM::Testing::Commit([&]
 		{
 			REQUIRE(Object->Rename(nullptr, OuterB));
 		});
 
-		REQUIRE(nullptr != Object);
-		REQUIRE(!Object->IsInOuter(OuterA));
-		REQUIRE(Object->IsInOuter(OuterB));
+		REQUIRE(Object->GetName() == Cat);
+		REQUIRE(Object->GetOuter() == OuterB);
 	}
 
-	SECTION("Abort")
+	SECTION("Abort(Rename(Object))")
+	{
+		AutoRTFM::Testing::Abort([&]
+		{
+			REQUIRE(Object->Rename(nullptr, OuterB));
+			AutoRTFM::AbortTransaction();
+		});
+
+		REQUIRE(Object->GetName() == Cat);
+		REQUIRE(Object->GetOuter() == OuterA);
+	}
+	
+	SECTION("Commit(Rename(Object), Rename(Object))")
 	{
-		UMyAutoRTFMTestObject* OuterA = NewObject<UMyAutoRTFMTestObject>();
-		UMyAutoRTFMTestObject* OuterB = NewObject<UMyAutoRTFMTestObject>();
-		UMyAutoRTFMTestObject* Object = NewObject<UMyAutoRTFMTestObject>(OuterA);
-		REQUIRE(Object->IsInOuter(OuterA));
+		AutoRTFM::Testing::Commit([&]
+		{
+			REQUIRE(Object->Rename(nullptr, OuterB));
+			REQUIRE(Object->Rename(nullptr, OuterC));
+		});
 
+		REQUIRE(Object->GetName() == Cat);
+		REQUIRE(Object->GetOuter() == OuterC);
+	}
+
+	SECTION("Abort(Rename(Object), Rename(Object))")
+	{
 		AutoRTFM::Testing::Abort([&]
 		{
 			REQUIRE(Object->Rename(nullptr, OuterB));
+			REQUIRE(Object->Rename(nullptr, OuterC));
 			AutoRTFM::AbortTransaction();
 		});
 
-		REQUIRE(nullptr != Object);
-		REQUIRE(Object->IsInOuter(OuterA));
-		REQUIRE(!Object->IsInOuter(OuterB));
+		REQUIRE(Object->GetName() == Cat);
+		REQUIRE(Object->GetOuter() == OuterA);
+	}
+
+	SECTION("Commit(Rename(Object), Commit(Rename(Object)))")
+	{
+		AutoRTFM::Testing::Commit([&]
+		{
+			REQUIRE(Object->Rename(nullptr, OuterB));
+			AutoRTFM::Testing::Commit([&]
+			{
+				REQUIRE(Object->Rename(nullptr, OuterC));
+			});
+		});
+
+		REQUIRE(Object->GetName() == Cat);
+		REQUIRE(Object->GetOuter() == OuterC);
+	}
+
+	SECTION("Commit(Rename(Object), Abort(Rename(Object)))")
+	{
+		AutoRTFM::Testing::Commit([&]
+		{
+			REQUIRE(Object->Rename(nullptr, OuterB));
+			AutoRTFM::Testing::Abort([&]
+			{
+				REQUIRE(Object->Rename(nullptr, OuterC));
+				AutoRTFM::AbortTransaction();
+			});
+		});
+
+		REQUIRE(Object->GetName() == Cat);
+		REQUIRE(Object->GetOuter() == OuterB);
+	}
+
+	SECTION("Abort(Rename(Object), Commit(Rename(Object)))")
+	{
+		AutoRTFM::Testing::Abort([&]
+		{
+			REQUIRE(Object->Rename(nullptr, OuterB));
+			AutoRTFM::Testing::Commit([&]
+			{
+				REQUIRE(Object->Rename(nullptr, OuterC));
+			});
+			AutoRTFM::AbortTransaction();
+		});
+
+		REQUIRE(Object->GetName() == Cat);
+		REQUIRE(Object->GetOuter() == OuterA);
+	}
+
+	SECTION("Abort(Rename(Object), Commit(Rename(Object)))")
+	{
+		AutoRTFM::Testing::Abort([&]
+		{
+			REQUIRE(Object->Rename(nullptr, OuterB));
+			AutoRTFM::Testing::Abort([&]
+			{
+				REQUIRE(Object->Rename(nullptr, OuterC));
+				AutoRTFM::AbortTransaction();
+			});
+			AutoRTFM::AbortTransaction();
+		});
+
+		REQUIRE(Object->GetName() == Cat);
+		REQUIRE(Object->GetOuter() == OuterA);
 	}
 }
 
diff --git a/Engine/Source/Runtime/CoreUObject/Private/UObject/UObjectBase.cpp b/Engine/Source/Runtime/CoreUObject/Private/UObject/UObjectBase.cpp
--- a/Engine/Source/Runtime/CoreUObject/Private/UObject/UObjectBase.cpp
+++ b/Engine/Source/Runtime/CoreUObject/Private/UObject/UObjectBase.cpp
@@ -5,6 +5,7 @@
 =============================================================================*/
 
 #include "UObject/UObjectBase.h"
+#include "AutoRTFM/AutoRTFM.h"
 #include "Misc/MessageDialog.h"
 #include "Misc/ConfigCacheIni.h"
 #include "HAL/IConsoleManager.h"
@@ -233,6 +234,7 @@ void UObjectBase::DeferredRegister(UClass *UClassStaticClass, const TCHAR* Packa
  *
  * @param Name name to assign to this uobject
  */
+UE_AUTORTFM_ALWAYS_OPEN
 void UObjectBase::AddObject(FName InName, EInternalObjectFlags InSetInternalFlags, int32 InInternalIndex, int32 InSerialNumber, FRemoteObjectId InRemoteId)
 {
 	NamePrivate = InName;
@@ -265,6 +267,27 @@ void UObjectBase::AddObject(FName InName, EInternalObjectFlags InSetInternalFlag
  */
 void UObjectBase::LowLevelRename(FName NewName,UObject *NewOuter)
 {
+	if (AutoRTFM::IsClosed())
+	{
+		FName OldName = NamePrivate;
+		UObject * OldOuter = OuterPrivate;
+		UE_AUTORTFM_ONABORT(this, NewName, NewOuter, OldName, OldOuter)
+		{
+			// Only rename if the new name / owner matches the current name / 
+			// owner, as the name may have been changed due to a whole-object
+			// revert.
+			if (NewName == NamePrivate && NewOuter == OuterPrivate)
+			{
+				LowLevelRename(OldName, OldOuter);
+			}
+		};
+		UE_AUTORTFM_OPEN
+		{
+			LowLevelRename(NewName, NewOuter);
+		};
+		return;
+	}
+
 #if STATS || ENABLE_STATNAMEDEVENTS_UOBJECT
 	((UObject*)this)->ResetStatID(); // reset the stat id since this thing now has a different name
 #endif
diff --git a/Engine/Source/Runtime/CoreUObject/Private/UObject/UObjectGlobals.cpp b/Engine/Source/Runtime/CoreUObject/Private/UObject/UObjectGlobals.cpp
--- a/Engine/Source/Runtime/CoreUObject/Private/UObject/UObjectGlobals.cpp
+++ b/Engine/Source/Runtime/CoreUObject/Private/UObject/UObjectGlobals.cpp
@@ -3744,6 +3744,17 @@ UObject* StaticAllocateObject
 				}
 #endif // !UE_WITH_REMOTE_OBJECT_HANDLE
 
+				// Record the memory of the object before destruction, so aborting will return the
+				// object to the previous object's state.
+				UE_AUTORTFM_OPEN_NO_VALIDATION
+				{
+					// We disable validation on this write because this spans the entire object's memory
+					// including fields that may be accessed only in the open (e.g. object flags which
+					// are set atomically), and mixing closed writes with open writes will trigger the
+					// memory validator.
+					AutoRTFM::RecordOpenWriteNoMemoryValidation((void *)Obj, TotalSize);
+				};
+
 				// Get the name before we start the destroy, as destroy renames it
 				FString OldName = Obj->GetFullName();
 
@@ -3814,19 +3825,6 @@ UObject* StaticAllocateObject
 
 		UE_AUTORTFM_OPEN_NO_VALIDATION
 		{
-			// If a this is a new object, we *do not* want to record the write as reverting the
-			// memory would bring the object to an uninitialized memory state, which will
-			// cause the GC to access uninitialized memory.
-			// If this is a recycled object, then we *do* want to record the write so an abort
-			// will return the object to the previous object's state.
-			if (bWasConstructedOnOldObject)
-			{
-				// We disable validation on this write because this spans the entire object's memory
-				// including fields that may be accessed only in the open (e.g. object flags which
-				// are set atomically), and mixing closed writes with open writes will trigger the
-				// memory validator.
-				AutoRTFM::RecordOpenWriteNoMemoryValidation((void *)Obj, TotalSize);
-			}
 			FMemory::Memzero((void *)Obj, TotalSize);
 			new ((void *)Obj) UObjectBase(const_cast<UClass*>(InClass), InFlags|RF_NeedInitialization, InternalSetFlags, InOuter, InName, OldIndex, OldSerialNumber == 0 ? SerialNumber : OldSerialNumber, 
 #if UE_WITH_REMOTE_OBJECT_HANDLE
diff --git a/Engine/Source/Runtime/CoreUObject/Private/UObject/UObjectHash.cpp b/Engine/Source/Runtime/CoreUObject/Private/UObject/UObjectHash.cpp
--- a/Engine/Source/Runtime/CoreUObject/Private/UObject/UObjectHash.cpp
+++ b/Engine/Source/Runtime/CoreUObject/Private/UObject/UObjectHash.cpp
@@ -5,6 +5,7 @@
 =============================================================================*/
 
 #include "UObject/UObjectHash.h"
+#include "Misc/AssertionMacros.h"
 #include "Misc/TransactionallySafeCriticalSection.h"
 #include "UObject/Class.h"
 #include "UObject/GarbageCollectionGlobals.h"
@@ -418,6 +419,7 @@ public:
 	{
 	}
 
+	UE_AUTORTFM_NOAUTORTFM  // This is not safe to be called from a closed transaction.
 	void ShrinkMaps()
 	{
 		const EParallelForFlags BaseFlags = FTaskGraphInterface::IsRunning() ? EParallelForFlags::None : EParallelForFlags::ForceSingleThread;
@@ -1056,6 +1058,7 @@ UObject* StaticFindObjectFastInternal(FRemoteObjectId RemoteId, EObjectFlags InE
 // Assumes that ThreadHash's critical is already locked
 FORCEINLINE static void AddToOuterMap(FUObjectHashTables& ThreadHash, UObjectBase* Object)
 {
+	ensureMsgf(!AutoRTFM::IsClosed(), TEXT("ObjectOuterMap must not be modified from a closed transaction"));
 	FHashBucket& Bucket = ThreadHash.ObjectOuterMap.FindOrAdd(Object->GetOuter());
 	checkSlow(!Bucket.Contains(Object)); // if it already exists, something is wrong with the external code
 	Bucket.Add(Object);
@@ -1123,6 +1126,7 @@ FORCEINLINE static void AddToRemoteIdMap(FUObjectHashTables& ThreadHash, UObject
 // Assumes that ThreadHash's critical is already locked
 FORCEINLINE static void RemoveFromOuterMap(FUObjectHashTables& ThreadHash, UObjectBase* Object)
 {
+	ensureMsgf(!AutoRTFM::IsClosed(), TEXT("ObjectOuterMap must not be modified from a closed transaction"));
 	UObject* Outer = UE::CoreUObject::Private::FObjectHandleUtils::GetNonAccessTrackedOuterNoResolve(Object);
 	FHashBucket& Bucket = ThreadHash.ObjectOuterMap.FindOrAdd(Outer);
 	int32 NumRemoved = Bucket.Remove(Object);
@@ -1670,6 +1674,7 @@ bool ClassHasInstancesAsyncLoading(const UClass* ClassToLookFor)
 	return false;
 }
 
+UE_AUTORTFM_NOAUTORTFM  // This is not safe to be called from a closed transaction.
 void HashObject(UObjectBase* Object)
 {
 	FName Name = Object->GetFName();
@@ -1724,6 +1729,7 @@ void HashObject(UObjectBase* Object)
  *
  * @param	Object		Object to remove from the hash tables
  */
+UE_AUTORTFM_NOAUTORTFM  // This is not safe to be called from a closed transaction.
 void UnhashObject(UObjectBase* Object)
 {
 	FName Name = Object->GetFName();
