[AutoRTFM] Add TStack - a minimal vector with inline storage.

Will be used as a TArray replacement, as AutoRTFM is being extracted from Core.
Currently unused, but tested. Future CLs will migrate AutoRTFM to use this.

diff --git a/Engine/Source/Programs/AutoRTFMTests/Private/ObjectLifetimeHelper.cpp b/Engine/Source/Programs/AutoRTFMTests/Private/ObjectLifetimeHelper.cpp
new file mode 100644
--- /dev/null
+++ b/Engine/Source/Programs/AutoRTFMTests/Private/ObjectLifetimeHelper.cpp
@@ -0,0 +1,66 @@
+// Copyright Epic Games, Inc. All Rights Reserved.
+
+#include "ObjectLifetimeHelper.h"
+#include "Catch2Includes.h"
+
+namespace AutoRTFMTestUtils
+{
+
+FObjectLifetimeHelper::FObjectLifetimeHelper()
+{
+	ConstructorCalls++;
+}
+
+FObjectLifetimeHelper::~FObjectLifetimeHelper()
+{
+	DestructorCalls++;
+}
+
+FObjectLifetimeHelper::FObjectLifetimeHelper(int Value) : Value{Value}
+{
+	ConstructorCalls++;
+}
+
+FObjectLifetimeHelper::FObjectLifetimeHelper(const FObjectLifetimeHelper& Other) : Value{Other.Value}, bIsValid(Other.bIsValid)
+{
+	ConstructorCalls++;
+	REQUIRE(bIsValid);
+}
+
+FObjectLifetimeHelper::FObjectLifetimeHelper(FObjectLifetimeHelper&& Other) : Value{Other.Value}, bIsValid(Other.bIsValid)
+{
+	Other.Value = 0;
+	Other.bIsValid = false;
+	ConstructorCalls++;
+	REQUIRE(bIsValid);
+}
+
+FObjectLifetimeHelper& FObjectLifetimeHelper::operator = (const FObjectLifetimeHelper& Other)
+{
+	Value = Other.Value;
+	bIsValid = Other.bIsValid;
+	REQUIRE(bIsValid);
+	return *this;
+}
+
+FObjectLifetimeHelper& FObjectLifetimeHelper::operator = (FObjectLifetimeHelper&& Other)
+{
+	Value = Other.Value;
+	bIsValid = Other.bIsValid;
+	Other.Value = 0;
+	Other.bIsValid = false;
+	REQUIRE(bIsValid);
+	return *this;
+}
+
+bool FObjectLifetimeHelper::operator == (FObjectLifetimeHelper Other) const
+{
+	REQUIRE(bIsValid);
+	REQUIRE(Other.bIsValid);
+	return Value == Other.Value;
+}
+
+size_t FObjectLifetimeHelper::ConstructorCalls = 0;
+size_t FObjectLifetimeHelper::DestructorCalls = 0;
+
+} // namespace FObjectLifetimeHelper
diff --git a/Engine/Source/Programs/AutoRTFMTests/Private/ObjectLifetimeHelper.h b/Engine/Source/Programs/AutoRTFMTests/Private/ObjectLifetimeHelper.h
new file mode 100644
--- /dev/null
+++ b/Engine/Source/Programs/AutoRTFMTests/Private/ObjectLifetimeHelper.h
@@ -0,0 +1,27 @@
+// Copyright Epic Games, Inc. All Rights Reserved.
+
+#pragma once
+
+namespace AutoRTFMTestUtils
+{
+	
+// A helper type that acts like an int, but also tracks number of ctors, dtors
+// and whether its used after move.
+struct FObjectLifetimeHelper
+{
+	FObjectLifetimeHelper();
+	~FObjectLifetimeHelper();
+	FObjectLifetimeHelper(int Value);
+	FObjectLifetimeHelper(const FObjectLifetimeHelper& Other);
+	FObjectLifetimeHelper(FObjectLifetimeHelper&& Other);
+	FObjectLifetimeHelper& operator = (const FObjectLifetimeHelper& Other);
+	FObjectLifetimeHelper& operator = (FObjectLifetimeHelper&& Other);
+	bool operator == (FObjectLifetimeHelper Other) const;
+
+	int Value = 0;
+	bool bIsValid = true;
+	static size_t ConstructorCalls;
+	static size_t DestructorCalls;
+};
+
+}  // namespace AutoRTFMTestUtils
diff --git a/Engine/Source/Programs/AutoRTFMTests/Private/StackTests.cpp b/Engine/Source/Programs/AutoRTFMTests/Private/StackTests.cpp
new file mode 100644
--- /dev/null
+++ b/Engine/Source/Programs/AutoRTFMTests/Private/StackTests.cpp
@@ -0,0 +1,450 @@
+// Copyright Epic Games, Inc. All Rights Reserved.
+
+#include "AutoRTFM/Stack.h"
+#include "Catch2Includes.h"
+#include "ObjectLifetimeHelper.h"
+
+using TrivialStack = AutoRTFM::TStack<int, 4>;
+using NonTrivialStack = AutoRTFM::TStack<AutoRTFMTestUtils::FObjectLifetimeHelper, 4>;
+
+TEMPLATE_TEST_CASE("Stack", "", TrivialStack, NonTrivialStack)
+{
+	REQUIRE(AutoRTFMTestUtils::FObjectLifetimeHelper::ConstructorCalls == 0);
+	REQUIRE(AutoRTFMTestUtils::FObjectLifetimeHelper::DestructorCalls == 0);
+
+	using StackType = TestType;
+	using ElementType = typename TestType::ElementType;
+
+	auto Check = [&](StackType& Stack, std::vector<ElementType> Expected)
+	{
+		// Check the reported count is as expected.
+		REQUIRE(Stack.Num() == Expected.size());
+		REQUIRE(Stack.IsEmpty() == Expected.empty());
+
+		if (!Expected.empty())
+		{
+			REQUIRE(Expected.back() == Stack.Back());
+			REQUIRE(Expected.front() == Stack.Front());
+			
+			// Check the Front() and Back() return a mutable reference.
+			Stack.Front() = 99;
+			REQUIRE(Stack.Front() == 99);
+			Stack.Front() = Expected.front();
+
+			Stack.Back() = 99;
+			REQUIRE(Stack.Back() == 99);
+			Stack.Back() = Expected.back();
+		}
+
+		for (size_t I = 0; I < Expected.size(); I++)
+		{
+			// Check the element is as expected.
+			REQUIRE(Stack[I] == Expected[I]);
+			// Check the index operator returns a mutable reference.
+			Stack[I] = 99;
+			REQUIRE(Stack[I] == 99);
+			Stack[I] = Expected[I];
+		}
+		// Check the non-const begin() / end().
+		{
+			size_t I = 0;
+			for (ElementType& Item : Stack)
+			{
+				REQUIRE(Item == Expected[I]);
+				I++;
+			}
+			REQUIRE(I == Expected.size());
+		}
+		// Check the const begin() / end().
+		{
+			size_t I = 0;
+			for (const ElementType& Item : const_cast<const StackType&>(Stack))
+			{
+				REQUIRE(Item == Expected[I]);
+				I++;
+			}
+			REQUIRE(I == Expected.size());
+		}
+	};
+
+	SECTION("Push, Pop")
+	{
+		StackType Stack;
+		Check(Stack, {});
+
+		Stack.Push(1);
+		ElementType* InlineAddress = &Stack[0];
+		Check(Stack, {1});
+
+		Stack.Push(2);
+		REQUIRE(&Stack[0] == InlineAddress);
+		Check(Stack, {1, 2});
+
+		Stack.Push(3);
+		REQUIRE(&Stack[0] == InlineAddress);
+		Check(Stack, {1, 2, 3});
+
+		Stack.Pop();
+		REQUIRE(&Stack[0] == InlineAddress);
+		Check(Stack, {1, 2});
+
+		Stack.Push(4);
+		REQUIRE(&Stack[0] == InlineAddress);
+		Check(Stack, {1, 2, 4});
+
+		Stack.Push(5);
+		REQUIRE(&Stack[0] == InlineAddress);
+		Check(Stack, {1, 2, 4, 5});
+
+		Stack.Push(6); // Spill inline -> heap
+		REQUIRE(&Stack[0] != InlineAddress);
+		Check(Stack, {1, 2, 4, 5, 6});
+
+		Stack.Pop();
+		REQUIRE(&Stack[0] != InlineAddress);
+		Check(Stack, {1, 2, 4, 5});
+	}
+
+	SECTION("PushAll")
+	{
+		StackType Target;
+		// PushAll() with empty target
+		{
+			StackType Source;
+			Source.Push(1);
+			Source.Push(2);
+			Target.PushAll(std::move(Source));
+			REQUIRE(Source.IsEmpty());
+			Check(Target, {1, 2});
+		}
+
+		// PushAll() with target holding inline data
+		{
+			StackType Source;
+			Source.Push(30);
+			Source.Push(40);
+			Target.PushAll(std::move(Source));
+			REQUIRE(Source.IsEmpty());
+			Check(Target, {1, 2, 30, 40});
+		}
+
+		// PushAll() with target spilling from inline -> heap.
+		{
+			StackType Source;
+			Source.Push(500);
+			Source.Push(600);
+			Source.Push(700);
+			Target.PushAll(std::move(Source));
+			REQUIRE(Source.IsEmpty());
+			Check(Target, {1, 2, 30, 40, 500, 600, 700});
+		}
+	}
+
+	SECTION("Clear / Reset")
+	{
+		StackType Stack;
+		Stack.Push(1);
+		ElementType* InlineAddress = &Stack[0];
+		Stack.Push(2);
+		Stack.Push(3);
+		Stack.Push(4);
+		Stack.Push(5);
+		SECTION("Clear")
+		{
+			Stack.Clear();
+			Check(Stack, {});
+			Stack.Push(100);
+			Check(Stack, {100});
+			REQUIRE(&Stack[0] != InlineAddress);
+		}
+		SECTION("Reset")
+		{
+			Stack.Reset();
+			Check(Stack, {});
+			Stack.Push(100);
+			Check(Stack, {100});
+			REQUIRE(&Stack[0] == InlineAddress);
+		}
+	}
+
+	SECTION("Copy Construct")
+	{
+		SECTION("Inline")
+		{
+			StackType Source;
+			Source.Push(1);
+			Source.Push(2);
+			Source.Push(3);
+			StackType Target(Source);
+			Check(Source, {1, 2, 3});
+			Check(Target, {1, 2, 3});
+		}
+
+		SECTION("Heap")
+		{
+			StackType Source;
+			Source.Push(1);
+			Source.Push(2);
+			Source.Push(3);
+			Source.Push(4);
+			Source.Push(5);
+			StackType Target(Source);
+			Check(Source, {1, 2, 3, 4, 5});
+			Check(Target, {1, 2, 3, 4, 5});
+		}
+	}
+
+	SECTION("Move Construct")
+	{
+		SECTION("Inline")
+		{
+			StackType Source;
+			Source.Push(1);
+			Source.Push(2);
+			Source.Push(3);
+			StackType Target(std::move(Source));
+			Check(Source, {});
+			Check(Target, {1, 2, 3});
+		}
+
+		SECTION("Heap")
+		{
+			StackType Source;
+			Source.Push(1);
+			Source.Push(2);
+			Source.Push(3);
+			Source.Push(4);
+			Source.Push(5);
+			StackType Target(std::move(Source));
+			Check(Source, {});
+			Check(Target, {1, 2, 3, 4, 5});
+		}
+	}
+
+	SECTION("Copy Assign")
+	{
+		StackType Source;
+		StackType Target;
+		auto CheckCopy = [&]()
+		{
+			SECTION("Source Empty")
+			{
+				Target = Source;
+				Check(Source, {});
+				Check(Target, {});
+			}
+			SECTION("Source Inline")
+			{
+				Source.Push(1);
+				Source.Push(2);
+				Source.Push(3);
+				Target = Source;
+				Check(Source, {1, 2, 3});
+				Check(Target, {1, 2, 3});
+			}
+			SECTION("Source Heap")
+			{
+				Source.Push(1);
+				Source.Push(2);
+				Source.Push(3);
+				Source.Push(4);
+				Source.Push(5);
+				Target = Source;
+				Check(Source, {1, 2, 3, 4, 5});
+				Check(Target, {1, 2, 3, 4, 5});
+			}
+		};
+
+		SECTION("Target Empty")
+		{
+			CheckCopy();
+		}
+		SECTION("Target Inline")
+		{
+			Target.Push(10);
+			Target.Push(20);
+			Target.Push(30);
+			CheckCopy();
+		}
+		SECTION("Target Heap")
+		{
+			Target.Push(10);
+			Target.Push(20);
+			Target.Push(30);
+			Target.Push(40);
+			Target.Push(50);
+			CheckCopy();
+		}
+	}
+
+	SECTION("Copy Assign Self")
+	{
+		StackType Stack;
+		StackType& Alias = Stack;
+		SECTION("Empty")
+		{
+			Stack = Alias;
+			Check(Stack, {});
+		}
+		SECTION("Inline")
+		{
+			Stack.Push(1);
+			Stack.Push(2);
+			Stack.Push(3);
+			Stack = Alias;
+			Check(Stack, {1, 2, 3});
+		}
+		SECTION("Heap")
+		{
+			Stack.Push(1);
+			Stack.Push(2);
+			Stack.Push(3);
+			Stack.Push(4);
+			Stack.Push(5);
+			Stack = Alias;
+			Check(Stack, {1, 2, 3, 4, 5});
+		}
+	}
+
+	SECTION("Move Assign")
+	{
+		StackType Source;
+		StackType Target;
+		auto CheckMove = [&]
+		{
+			SECTION("Source Empty")
+			{
+				Target = std::move(Source);
+				Check(Source, {});
+				Check(Target, {});
+			}
+			SECTION("Source Inline")
+			{
+				Source.Push(1);
+				Source.Push(2);
+				Source.Push(3);
+				Target = std::move(Source);
+				Check(Source, {});
+				Check(Target, {1, 2, 3});
+			}
+			SECTION("Source Heap")
+			{
+				Source.Push(1);
+				Source.Push(2);
+				Source.Push(3);
+				Source.Push(4);
+				Source.Push(5);
+				Target = std::move(Source);
+				Check(Source, {});
+				Check(Target, {1, 2, 3, 4, 5});
+			}
+		};
+
+		SECTION("Target Empty")
+		{
+			CheckMove();
+		}
+		SECTION("Target Inline")
+		{
+			Target.Push(10);
+			Target.Push(20);
+			Target.Push(30);
+			CheckMove();
+		}
+		SECTION("Target Heap")
+		{
+			Target.Push(10);
+			Target.Push(20);
+			Target.Push(30);
+			Target.Push(40);
+			Target.Push(50);
+			CheckMove();
+		}
+	}
+
+	SECTION("Move Assign Self")
+	{
+		StackType Stack;
+		StackType& Alias = Stack;
+		SECTION("Empty")
+		{
+			Stack = std::move(Alias);
+			Check(Stack, {});
+		}
+		SECTION("Inline")
+		{
+			Stack.Push(1);
+			Stack.Push(2);
+			Stack.Push(3);
+			Stack = std::move(Alias);
+			Check(Stack, {1, 2, 3});
+		}
+		SECTION("Heap")
+		{
+			Stack.Push(1);
+			Stack.Push(2);
+			Stack.Push(3);
+			Stack.Push(4);
+			Stack.Push(5);
+			Stack = std::move(Alias);
+			Check(Stack, {1, 2, 3, 4, 5});
+		}
+	}
+
+	SECTION("Soak")
+	{
+		StackType Stack;
+		std::vector<ElementType> Expected;
+		for (int I = 0; I < 10000; I++)
+		{
+			int Mod100 = (I * 15485863) % 100;
+			switch (Mod100)
+			{
+				case 0:
+				{
+					Stack.Clear();
+					Expected.clear();
+					break;
+				}
+				case 1:
+				{	
+					Stack.Reset();
+					Expected.clear();
+					break;
+				}
+				case 2:
+				{	
+					StackType Copy(Stack);
+					Stack = Copy;
+					break;
+				}
+				case 3:
+				{	
+					StackType Move(std::move(Stack));
+					Stack = std::move(Move);
+					break;
+				}
+				default:
+				{
+					if (Mod100 > 40 || Expected.empty())
+					{
+						Stack.Push(I);
+						Expected.push_back(I);
+						Check(Stack, Expected);
+					}
+					else
+					{
+						Stack.Pop();
+						Expected.pop_back();
+						Check(Stack, Expected);
+					}
+					break;
+				}
+			}
+		}
+	}
+
+	REQUIRE(AutoRTFMTestUtils::FObjectLifetimeHelper::ConstructorCalls == AutoRTFMTestUtils::FObjectLifetimeHelper::DestructorCalls);
+	AutoRTFMTestUtils::FObjectLifetimeHelper::ConstructorCalls = 0;
+	AutoRTFMTestUtils::FObjectLifetimeHelper::DestructorCalls = 0;
+}
diff --git a/Engine/Source/Runtime/Core/Private/AutoRTFM/Stack.h b/Engine/Source/Runtime/Core/Private/AutoRTFM/Stack.h
new file mode 100644
--- /dev/null
+++ b/Engine/Source/Runtime/Core/Private/AutoRTFM/Stack.h
@@ -0,0 +1,308 @@
+// Copyright Epic Games, Inc. All Rights Reserved.
+
+#pragma once
+
+#include "ExternAPI.h"
+#include "Utils.h"
+#include <cstddef>
+#include <type_traits>
+
+namespace AutoRTFM
+{
+
+// Enumerator used to control validity assertions for TStack (bounds checking, 
+// etc).
+enum class EStackValidation
+{
+	Enabled,
+	Disabled
+};
+
+// A stack container with an inline fixed capacity that once exceeded spills to
+// the heap.
+//
+// Notes:
+//  * Heap allocated memory is not automatically freed when popping elements.
+//    Only calling Reset() or destructing the stack will free heap allocated
+//    memory.
+//  * This class is not relocatable and so is not safe to use in UE's containers
+//    which require elements to be relocatable.
+//
+// Template parameters:
+//   T - the element type. Must support copy and move constructors.
+//   InlineCapacity - the number of elements that can be held before spilling to
+//                    the heap.
+//   Validation - if Disabled, do not perform validity assertions.
+template<typename T, size_t InlineCapacity, EStackValidation Validation = EStackValidation::Enabled>
+class TStack
+{
+public:
+	using ElementType = T;
+
+	// Constructor.
+	TStack() = default;
+
+	// Copy constructor.
+	TStack(const TStack& Other)
+	{
+		CopyFrom(Other);
+	}
+
+	// Move constructor.
+	TStack(TStack&& Other)
+	{
+		MoveFrom(std::move(Other));
+	}
+
+	// Destructor.
+	// Frees all the memory allocated by the allocator.
+	~TStack()
+	{
+		Reset();
+	}
+
+	// Copy assignment operator
+	TStack& operator=(const TStack& Other)
+	{
+		if (&Other != this)
+		{
+			Clear();
+			CopyFrom(Other);
+		}
+		return *this;
+	}
+
+	// Move assignment operator
+	TStack& operator=(TStack&& Other)
+	{
+		if (&Other != this)
+		{
+			Clear();
+			MoveFrom(std::move(Other));
+		}
+		return *this;
+	}
+	
+	// Clears all the items from the vector, preserving the capacity.
+	void Clear()
+	{
+		for (T& Item : *this)
+		{
+			Item.~T();
+		}
+		Count = 0;
+	}
+
+	// Clears all the items from the vector, freeing all heap allocations and
+	// resetting the capacity to InlineCapacity
+	void Reset()
+	{
+		Clear();
+		if (Data != InlineData())
+		{
+			Free(Data);
+			Data = reinterpret_cast<T*>(InlineDataBuffer);
+		}
+		Capacity = InlineCapacity;
+	}
+
+	// Pushes a new item on to the stack.
+	inline void Push(const T& Item)
+	{
+		if (Count >= Capacity)
+		{
+			Reserve(Capacity * 2);
+		}
+
+		new (Data + Count) T(Item);
+		++Count;
+	}
+
+	// Moves all the items from Other to this stack.
+	// Other will hold no elements after calling.
+	inline void PushAll(TStack&& Other)
+	{
+		size_t NewCount = Count + Other.Count;
+		if (NewCount >= Capacity)
+		{
+			Reserve(std::max<size_t>(NewCount, Capacity * 2));
+		}
+
+		if constexpr (std::is_trivial_v<T>)
+		{
+			memcpy(Data + Count, Other.Data, Other.Count * sizeof(T));
+		}
+		else
+		{
+			for (size_t I = 0, N = Other.Count; I < N; I++)
+			{
+				new (Data + Count + I) T(std::move(Other.Data[I]));
+				Other.Data[I].~T();
+			}
+		}
+		Count = NewCount;
+		Other.SetToInitialState();
+	}
+
+	// Removes the last item on the stack.
+	inline void Pop()
+	{
+		ASSERT(Validation == EStackValidation::Disabled || Count > 0);
+		--Count;
+		Data[Count].~T();
+	}
+
+	// Reserves memory for at NewCapacity items.
+	inline void Reserve(size_t NewCapacity)
+	{
+		if (NewCapacity <= Capacity)
+		{
+			return; // Already has space for the new capacity
+		}
+
+		if constexpr (std::is_trivial_v<T>)
+		{
+			if (Data == InlineData())
+			{
+				T* NewData = reinterpret_cast<T*>(AutoRTFM::Allocate(NewCapacity * sizeof(T), alignof(T)));
+				memcpy(NewData, InlineDataBuffer, Count * sizeof(T));
+				Data = NewData;
+			}
+			else
+			{
+				Data = reinterpret_cast<T*>(AutoRTFM::Reallocate(Data, NewCapacity * sizeof(T), alignof(T)));
+			}
+		}
+		else
+		{
+			T* NewData = reinterpret_cast<T*>(AutoRTFM::Allocate(NewCapacity * sizeof(T), alignof(T)));
+			for (size_t I = 0, N = Count; I < N; I++)
+			{
+				new (NewData + I) T(std::move(Data[I]));
+				Data[I].~T();
+			}
+			if (Data != InlineData())
+			{
+				AutoRTFM::Free(Data);
+			}
+			Data = NewData;
+		}
+
+		Capacity = NewCapacity;
+	}
+
+	inline size_t Num() const { return Count; }
+
+	inline bool IsEmpty() const { return Count == 0; }
+
+	T& operator[](size_t Index)
+	{
+		ASSERT(Validation == EStackValidation::Disabled || Index < Count);
+		return Data[Index];
+	}
+
+	const T& operator[](size_t Index) const
+	{
+		ASSERT(Validation == EStackValidation::Disabled || Index < Count);
+		return Data[Index];
+	}
+
+	T& Front()
+	{
+		ASSERT(Validation == EStackValidation::Disabled || Count > 0);
+		return Data[0];
+	}
+
+	const T& Front() const
+	{
+		ASSERT(Validation == EStackValidation::Disabled || Count > 0);
+		return Data[0];
+	}
+
+	T& Back()
+	{
+		ASSERT(Validation == EStackValidation::Disabled || Count > 0);
+		return Data[Count - 1];
+	}
+
+	const T& Back() const
+	{
+		ASSERT(Validation == EStackValidation::Disabled || Count > 0);
+		return Data[Count - 1];
+	}
+
+	T* begin() { return Data; }
+	const T* begin() const { return Data; }
+	T* end() { return Data + Count; }
+	const T* end() const { return Data + Count; }
+
+private:
+	// Copies the data from Other to this stack.
+	// This stack must be empty before calling.
+	void CopyFrom(const TStack& Other)
+	{
+		ASSERT(IsEmpty());
+		Reserve(Other.Count);
+		if constexpr (std::is_trivial_v<T>)
+		{
+			memcpy(Data, Other.Data, Other.Count * sizeof(T));
+		}
+		else
+		{
+			for (size_t I = 0, N = Other.Count; I < N; I++)
+			{
+				new (Data + I) T(Other.Data[I]);
+			}
+		}
+		Count = Other.Count;
+	}
+
+	// Moves the data from Other to this stack.
+	// This stack must be empty before calling.
+	void MoveFrom(TStack&& Other)
+	{
+		ASSERT(IsEmpty());
+		if (Other.Data == Other.InlineData())
+		{
+			if constexpr (std::is_trivial_v<T>)
+			{
+				memcpy(Data, Other.Data, Other.Count * sizeof(T));
+			}
+			else
+			{
+				for (size_t I = 0, N = Other.Count; I < N; I++)
+				{
+					new (Data + I) T(std::move(Other.Data[I]));
+					Other.Data[I].~T();
+				}
+			}
+		}
+		else
+		{
+			// Steal the heap allocation from Other.
+			Data = Other.Data;
+			Capacity = Other.Capacity;
+		}
+		Count = Other.Count;
+		Other.SetToInitialState();
+	}
+
+	// Sets Data, Capacity and Count to the initially default-constructed values.
+	// Warning: This does not free or destruct any elements held by this stack.
+	void SetToInitialState()
+	{
+		Data = InlineData();
+		Capacity = InlineCapacity;
+		Count = 0;
+	}
+
+	T* InlineData() { return reinterpret_cast<T*>(InlineDataBuffer); }
+	const T* InlineData() const { return reinterpret_cast<const T*>(InlineDataBuffer); }
+
+	alignas(T) std::byte InlineDataBuffer[sizeof(T) * InlineCapacity];
+	T* Data = reinterpret_cast<T*>(InlineDataBuffer);
+	size_t Capacity = InlineCapacity;
+	size_t Count = 0;
+};
+
+}
