Make `ParallelFor` single-threaded when run from within a transaction.

diff --git a/Engine/Source/Programs/AutoRTFMTests/Private/UECoreTests.cpp b/Engine/Source/Programs/AutoRTFMTests/Private/UECoreTests.cpp
--- a/Engine/Source/Programs/AutoRTFMTests/Private/UECoreTests.cpp
+++ b/Engine/Source/Programs/AutoRTFMTests/Private/UECoreTests.cpp
@@ -3,6 +3,7 @@
 #include "AutoRTFM/AutoRTFM.h"
 #include "Catch2Includes.h"
 #include "Async/TransactionallySafeMutex.h"
+#include "Async/ParallelFor.h"
 #include "AssetRegistry/AssetDataTagMap.h"
 #include "AutoRTFM/AutoRTFMTesting.h"
 #include "Delegates/IDelegateInstance.h"
@@ -3565,3 +3566,21 @@ TEST_CASE("UECore.GetCurrentProcessId")
 		REQUIRE(Outer == FPlatformProcess::GetCurrentProcessId());
 	});
 }
+
+TEST_CASE("UECore.ParallelFor")
+{
+	const int32 Parallelism = 2;
+	FTransactionallySafeMutex Mutex;
+	int32 Count = 0;
+
+	AutoRTFM::Testing::Commit([&]
+	{
+		ParallelFor(Parallelism, [&](int32 ThreadId)
+		{
+			UE::TScopeLock _(Mutex);
+			Count += 1;
+		});
+	});
+
+	REQUIRE(Count == Parallelism);
+}
\ No newline at end of file
diff --git a/Engine/Source/Runtime/Core/Public/Async/ParallelFor.h b/Engine/Source/Runtime/Core/Public/Async/ParallelFor.h
--- a/Engine/Source/Runtime/Core/Public/Async/ParallelFor.h
+++ b/Engine/Source/Runtime/Core/Public/Async/ParallelFor.h
@@ -6,6 +6,7 @@
 
 #pragma once
 
+#include "AutoRTFM/AutoRTFM.h"
 #include "Async/Fundamental/Scheduler.h"
 #include "Async/Fundamental/Task.h"
 #include "Async/TaskGraphInterfaces.h"
@@ -125,6 +126,12 @@ namespace ParallelForImpl
 		TRACE_CPUPROFILER_EVENT_SCOPE(ParallelFor);
 		check(Num >= 0);
 
+		// If we are executing the parallel for from within a transaction, run it single-threaded.
+		if (AutoRTFM::IsClosed())
+		{
+			Flags = EParallelForFlags::ForceSingleThread;
+		}
+
 		int32 NumWorkers = GetNumberOfThreadTasks(Num, MinBatchSize, Flags);
 
 		if (!Contexts.IsEmpty())
