[AutoRTFM] Fix linux build: Add missing #include for use of size_t

diff --git a/Engine/Source/Programs/AutoRTFMTests/Private/ObjectLifetimeHelper.h b/Engine/Source/Programs/AutoRTFMTests/Private/ObjectLifetimeHelper.h
--- a/Engine/Source/Programs/AutoRTFMTests/Private/ObjectLifetimeHelper.h
+++ b/Engine/Source/Programs/AutoRTFMTests/Private/ObjectLifetimeHelper.h
@@ -2,6 +2,8 @@
 
 #pragma once
 
+#include <cstddef>
+
 namespace AutoRTFMTestUtils
 {
 	
