[AutoRTFM] Simplify test code by using AutoRTFM::Testing helper functions.

These functions automatically open a Transact and verify that its transaction-result matches our expectation. This saves a few hundred lines of boilerplate overall.

diff --git a/Engine/Source/Programs/AutoRTFMTests/Private/UECoreTests.cpp b/Engine/Source/Programs/AutoRTFMTests/Private/UECoreTests.cpp
--- a/Engine/Source/Programs/AutoRTFMTests/Private/UECoreTests.cpp
+++ b/Engine/Source/Programs/AutoRTFMTests/Private/UECoreTests.cpp
@@ -42,13 +42,12 @@ TEST_CASE("UECore.FDelegateHandle")
 
 	SECTION("With Abort")
 	{
-		AutoRTFM::ETransactionResult Result = AutoRTFM::Transact([&]()
+		AutoRTFM::Testing::Abort([&]()
 		{
 			Handle = FDelegateHandle(FDelegateHandle::GenerateNewHandle);
 			AutoRTFM::AbortTransaction();
 		});
 
-		REQUIRE(AutoRTFM::ETransactionResult::AbortedByRequest == Result);
 		REQUIRE(!Handle.IsValid());
 	}
 
@@ -56,12 +55,11 @@ TEST_CASE("UECore.FDelegateHandle")
 
 	SECTION("With Commit")
 	{
-		AutoRTFM::ETransactionResult Result = AutoRTFM::Transact([&]()
+		AutoRTFM::Testing::Commit([&]()
 		{
 			Handle = FDelegateHandle(FDelegateHandle::GenerateNewHandle);
 		});
 
-		REQUIRE(AutoRTFM::ETransactionResult::Committed == Result);
 		REQUIRE(Handle.IsValid());
 	}
 }
@@ -83,12 +81,11 @@ TEST_CASE("UECore.TThreadSingleton")
 		uintptr_t Data = 0x12345678abcdef00;
 		memcpy(&Singleton, &Data, sizeof(Singleton));
 
-		AutoRTFM::ETransactionResult Result = AutoRTFM::Transact([&]()
+		AutoRTFM::Testing::Commit([&]()
 		{
 			Singleton = TThreadSingleton<MyStruct>::TryGet();
 		});
 
-		REQUIRE(AutoRTFM::ETransactionResult::Committed == Result);
 		REQUIRE(nullptr == Singleton);
 	}
 
@@ -96,14 +93,12 @@ TEST_CASE("UECore.TThreadSingleton")
 	{
 		MALLOCLEAK_IGNORE_SCOPE(); // TThreadSingleton will appear as a leak.
 
-		AutoRTFM::ETransactionResult Result = AutoRTFM::Transact([&]()
-			{
-				TThreadSingleton<MyStruct>::Get().I = 42;
-				TThreadSingleton<MyStruct>::Get().F = 42.0f;
-				AutoRTFM::AbortTransaction();
-			});
-
-		REQUIRE(AutoRTFM::ETransactionResult::AbortedByRequest == Result);
+		AutoRTFM::Testing::Abort([&]()
+		{
+			TThreadSingleton<MyStruct>::Get().I = 42;
+			TThreadSingleton<MyStruct>::Get().F = 42.0f;
+			AutoRTFM::AbortTransaction();
+		});
 
 		// The singleton *will remain* initialized though, even though we got it in
 		// a transaction, because we have to do the singleton creation in the open.
@@ -115,13 +110,11 @@ TEST_CASE("UECore.TThreadSingleton")
 		REQUIRE(0 == TThreadSingleton<MyStruct>::Get().I);
 		REQUIRE(0.0f == TThreadSingleton<MyStruct>::Get().F);
 
-		Result = AutoRTFM::Transact([&]()
-			{
-				TThreadSingleton<MyStruct>::Get().I = 42;
-				TThreadSingleton<MyStruct>::Get().F = 42.0f;
-			});
-
-		REQUIRE(AutoRTFM::ETransactionResult::Committed == Result);
+		AutoRTFM::Testing::Commit([&]()
+		{
+			TThreadSingleton<MyStruct>::Get().I = 42;
+			TThreadSingleton<MyStruct>::Get().F = 42.0f;
+		});
 
 		REQUIRE(42 == TThreadSingleton<MyStruct>::Get().I);
 		REQUIRE(42.0f == TThreadSingleton<MyStruct>::Get().F);
@@ -133,12 +126,11 @@ TEST_CASE("UECore.TThreadSingleton")
 
 		MyStruct* Singleton = nullptr;
 
-		AutoRTFM::ETransactionResult Result = AutoRTFM::Transact([&]()
-			{
-				Singleton = TThreadSingleton<MyStruct>::TryGet();
-			});
+		AutoRTFM::Testing::Commit([&]()
+		{
+			Singleton = TThreadSingleton<MyStruct>::TryGet();
+		});
 
-		REQUIRE(AutoRTFM::ETransactionResult::Committed == Result);
 		REQUIRE(nullptr != Singleton);
 	}
 }
@@ -165,23 +157,19 @@ TEST_CASE("UECore.FTextHistory")
 
 	SECTION("With Abort")
 	{
-		AutoRTFM::ETransactionResult Result = AutoRTFM::Transact([&]()
-			{
-				History.UpdateDisplayStringIfOutOfDate();
-				AutoRTFM::AbortTransaction();
-			});
-
-		REQUIRE(AutoRTFM::ETransactionResult::AbortedByRequest == Result);
+		AutoRTFM::Testing::Abort([&]()
+		{
+			History.UpdateDisplayStringIfOutOfDate();
+			AutoRTFM::AbortTransaction();
+		});
 	}
 
 	SECTION("With Commit")
 	{
-		AutoRTFM::ETransactionResult Result = AutoRTFM::Transact([&]()
-			{
-				History.UpdateDisplayStringIfOutOfDate();
-			});
-
-		REQUIRE(AutoRTFM::ETransactionResult::Committed == Result);
+		AutoRTFM::Testing::Commit([&]()
+		{
+			History.UpdateDisplayStringIfOutOfDate();
+		});
 	}
 }
 
@@ -196,32 +184,30 @@ TEST_CASE("UECore.FCustomVersionContainer")
 
 	SECTION("With Abort")
 	{
-		AutoRTFM::ETransactionResult Result = AutoRTFM::Transact([&]()
-			{
-				// The first time the version will be new.
-				Container.SetVersionUsingRegistry(Guid);
+		AutoRTFM::Testing::Abort([&]()
+		{
+			// The first time the version will be new.
+			Container.SetVersionUsingRegistry(Guid);
 
-				// The second time we should hit the cache the first one created.
-				Container.SetVersionUsingRegistry(Guid);
-				AutoRTFM::AbortTransaction();
-			});
+			// The second time we should hit the cache the first one created.
+			Container.SetVersionUsingRegistry(Guid);
+			AutoRTFM::AbortTransaction();
+		});
 
-		REQUIRE(AutoRTFM::ETransactionResult::AbortedByRequest == Result);
 		REQUIRE(nullptr == Container.GetVersion(Guid));
 	}
 
 	SECTION("With Commit")
 	{
-		AutoRTFM::ETransactionResult Result = AutoRTFM::Transact([&]()
-			{
-				// The first time the version will be new.
-				Container.SetVersionUsingRegistry(Guid);
+		AutoRTFM::Testing::Commit([&]()
+		{
+			// The first time the version will be new.
+			Container.SetVersionUsingRegistry(Guid);
 
-				// The second time we should hit the cache the first one created.
-				Container.SetVersionUsingRegistry(Guid);
-			});
+			// The second time we should hit the cache the first one created.
+			Container.SetVersionUsingRegistry(Guid);
+		});
 
-		REQUIRE(AutoRTFM::ETransactionResult::Committed == Result);
 		REQUIRE(nullptr != Container.GetVersion(Guid));
 	}
 }
@@ -234,24 +220,22 @@ TEST_CASE("UECore.FName")
 
 		SECTION("With Abort")
 		{
-			AutoRTFM::ETransactionResult Result = AutoRTFM::Transact([&]()
-				{
-					Name = FName(EName::Timer);
-					AutoRTFM::AbortTransaction();
-				});
+			AutoRTFM::Testing::Abort([&]()
+			{
+				Name = FName(EName::Timer);
+				AutoRTFM::AbortTransaction();
+			});
 
-			REQUIRE(AutoRTFM::ETransactionResult::AbortedByRequest == Result);
 			REQUIRE(Name.IsNone());
 		}
 
 		SECTION("With Commit")
 		{
-			AutoRTFM::ETransactionResult Result = AutoRTFM::Transact([&]()
-				{
-					Name = FName(EName::Timer);
-				});
+			AutoRTFM::Testing::Commit([&]()
+			{
+				Name = FName(EName::Timer);
+			});
 
-			REQUIRE(AutoRTFM::ETransactionResult::Committed == Result);
 			REQUIRE(EName::Timer == *Name.ToEName());
 		}
 	}
@@ -262,13 +246,12 @@ TEST_CASE("UECore.FName")
 
 		SECTION("With Abort")
 		{
-			AutoRTFM::ETransactionResult Result = AutoRTFM::Transact([&]()
-				{
-					Name = FName(TEXT("WOWWEE"), 42);
-					AutoRTFM::AbortTransaction();
-				});
+			AutoRTFM::Testing::Abort([&]()
+			{
+				Name = FName(TEXT("WOWWEE"), 42);
+				AutoRTFM::AbortTransaction();
+			});
 
-			REQUIRE(AutoRTFM::ETransactionResult::AbortedByRequest == Result);
 			REQUIRE(Name.IsNone());
 		}
 
@@ -294,12 +277,11 @@ TEST_CASE("UECore.FName")
 
 		SECTION("With Commit")
 		{
-			AutoRTFM::ETransactionResult Result = AutoRTFM::Transact([&]()
-				{
-					Name = FName(TEXT("WOWWEE"), 42);
-				});
+			AutoRTFM::Testing::Commit([&]()
+			{
+				Name = FName(TEXT("WOWWEE"), 42);
+			});
 
-			REQUIRE(AutoRTFM::ETransactionResult::Committed == Result);
 			REQUIRE(TEXT("WOWWEE") == Name.GetPlainNameString());
 			REQUIRE(42 == Name.GetNumber());
 		}
@@ -312,24 +294,21 @@ TEST_CASE("UECore.STATIC_FUNCTION_FNAME")
 
 	SECTION("With Abort")
 	{
-		AutoRTFM::ETransactionResult Result = AutoRTFM::Transact([&]()
-			{
-				Name = STATIC_FUNCTION_FNAME(TEXT("WOWWEE"));
-				AutoRTFM::AbortTransaction();
-			});
+		AutoRTFM::Testing::Abort([&]()
+		{
+			Name = STATIC_FUNCTION_FNAME(TEXT("WOWWEE"));
+			AutoRTFM::AbortTransaction();
+		});
 
-		REQUIRE(AutoRTFM::ETransactionResult::AbortedByRequest == Result);
 		REQUIRE(Name.IsNone());
 	}
 
 	SECTION("With Commit")
 	{
-		AutoRTFM::ETransactionResult Result = AutoRTFM::Transact([&]()
-			{
-				Name = STATIC_FUNCTION_FNAME(TEXT("WOWWEE"));
-			});
-
-		REQUIRE(AutoRTFM::ETransactionResult::Committed == Result);
+		AutoRTFM::Testing::Commit([&]()
+		{
+			Name = STATIC_FUNCTION_FNAME(TEXT("WOWWEE"));
+		});
 	}
 }
 
@@ -341,24 +320,22 @@ TEST_CASE("UECore.TIntrusiveReferenceController")
 
 		SECTION("With Abort")
 		{
-			AutoRTFM::ETransactionResult Result = AutoRTFM::Transact([&]()
-				{
-					Controller.AddSharedReference();
-					AutoRTFM::AbortTransaction();
-				});
+			AutoRTFM::Testing::Abort([&]()
+			{
+				Controller.AddSharedReference();
+				AutoRTFM::AbortTransaction();
+			});
 
-			REQUIRE(AutoRTFM::ETransactionResult::AbortedByRequest == Result);
 			REQUIRE(1 == Controller.GetSharedReferenceCount());
 		}
 
 		SECTION("With Commit")
 		{
-			AutoRTFM::ETransactionResult Result = AutoRTFM::Transact([&]()
-				{
-					Controller.AddSharedReference();
-				});
+			AutoRTFM::Testing::Commit([&]()
+			{
+				Controller.AddSharedReference();
+			});
 
-			REQUIRE(AutoRTFM::ETransactionResult::Committed == Result);
 			REQUIRE(2 == Controller.GetSharedReferenceCount());
 		}
 	}
@@ -369,24 +346,22 @@ TEST_CASE("UECore.TIntrusiveReferenceController")
 
 		SECTION("With Abort")
 		{
-			AutoRTFM::ETransactionResult Result = AutoRTFM::Transact([&]()
-				{
-					Controller.AddWeakReference();
-					AutoRTFM::AbortTransaction();
-				});
+			AutoRTFM::Testing::Abort([&]()
+			{
+				Controller.AddWeakReference();
+				AutoRTFM::AbortTransaction();
+			});
 
-			REQUIRE(AutoRTFM::ETransactionResult::AbortedByRequest == Result);
 			REQUIRE(1 == Controller.WeakReferenceCount);
 		}
 
 		SECTION("With Commit")
 		{
-			AutoRTFM::ETransactionResult Result = AutoRTFM::Transact([&]()
-				{
-					Controller.AddWeakReference();
-				});
+			AutoRTFM::Testing::Commit([&]()
+			{
+				Controller.AddWeakReference();
+			});
 
-			REQUIRE(AutoRTFM::ETransactionResult::Committed == Result);
 			REQUIRE(2 == Controller.WeakReferenceCount);
 		}
 	}
@@ -399,24 +374,22 @@ TEST_CASE("UECore.TIntrusiveReferenceController")
 
 			SECTION("With Abort")
 			{
-				AutoRTFM::ETransactionResult Result = AutoRTFM::Transact([&]()
-					{
-						Controller.ConditionallyAddSharedReference();
-						AutoRTFM::AbortTransaction();
-					});
+				AutoRTFM::Testing::Abort([&]()
+				{
+					Controller.ConditionallyAddSharedReference();
+					AutoRTFM::AbortTransaction();
+				});
 
-				REQUIRE(AutoRTFM::ETransactionResult::AbortedByRequest == Result);
 				REQUIRE(1 == Controller.GetSharedReferenceCount());
 			}
 
 			SECTION("With Commit")
 			{
-				AutoRTFM::ETransactionResult Result = AutoRTFM::Transact([&]()
-					{
-						Controller.ConditionallyAddSharedReference();
-					});
+				AutoRTFM::Testing::Commit([&]()
+				{
+					Controller.ConditionallyAddSharedReference();
+				});
 
-				REQUIRE(AutoRTFM::ETransactionResult::Committed == Result);
 				REQUIRE(2 == Controller.GetSharedReferenceCount());
 			}
 		}
@@ -432,24 +405,22 @@ TEST_CASE("UECore.TIntrusiveReferenceController")
 
 			SECTION("With Abort")
 			{
-				AutoRTFM::ETransactionResult Result = AutoRTFM::Transact([&]()
-					{
-						Controller.ConditionallyAddSharedReference();
-						AutoRTFM::AbortTransaction();
-					});
+				AutoRTFM::Testing::Abort([&]()
+				{
+					Controller.ConditionallyAddSharedReference();
+					AutoRTFM::AbortTransaction();
+				});
 
-				REQUIRE(AutoRTFM::ETransactionResult::AbortedByRequest == Result);
 				REQUIRE(0 == Controller.GetSharedReferenceCount());
 			}
 
 			SECTION("With Commit")
 			{
-				AutoRTFM::ETransactionResult Result = AutoRTFM::Transact([&]()
-					{
-						Controller.ConditionallyAddSharedReference();
-					});
+				AutoRTFM::Testing::Commit([&]()
+				{
+					Controller.ConditionallyAddSharedReference();
+				});
 
-				REQUIRE(AutoRTFM::ETransactionResult::Committed == Result);
 				REQUIRE(0 == Controller.GetSharedReferenceCount());
 			}
 		}
@@ -463,13 +434,12 @@ TEST_CASE("UECore.TIntrusiveReferenceController")
 		{
 			int32 Count = 0;
 
-			AutoRTFM::ETransactionResult Result = AutoRTFM::Transact([&]()
-				{
-					Count = Controller.GetSharedReferenceCount();
-					AutoRTFM::AbortTransaction();
-				});
+			AutoRTFM::Testing::Abort([&]()
+			{
+				Count = Controller.GetSharedReferenceCount();
+				AutoRTFM::AbortTransaction();
+			});
 
-			REQUIRE(AutoRTFM::ETransactionResult::AbortedByRequest == Result);
 			REQUIRE(0 == Count);
 		}
 
@@ -477,12 +447,11 @@ TEST_CASE("UECore.TIntrusiveReferenceController")
 		{
 			int32 Count = 0;
 
-			AutoRTFM::ETransactionResult Result = AutoRTFM::Transact([&]()
-				{
-					Count = Controller.GetSharedReferenceCount();
-				});
+			AutoRTFM::Testing::Commit([&]()
+			{
+				Count = Controller.GetSharedReferenceCount();
+			});
 
-			REQUIRE(AutoRTFM::ETransactionResult::Committed == Result);
 			REQUIRE(1 == Count);
 		}
 	}
@@ -495,12 +464,11 @@ TEST_CASE("UECore.TIntrusiveReferenceController")
 		{
 			bool Unique = false;
 
-			AutoRTFM::ETransactionResult Result = AutoRTFM::Transact([&]()
-				{
-					Unique = Controller.IsUnique();
-				});
+			AutoRTFM::Testing::Commit([&]()
+			{
+				Unique = Controller.IsUnique();
+			});
 
-			REQUIRE(AutoRTFM::ETransactionResult::Committed == Result);
 			REQUIRE(Unique);
 		}
 
@@ -511,12 +479,11 @@ TEST_CASE("UECore.TIntrusiveReferenceController")
 
 			bool Unique = true;
 
-			AutoRTFM::ETransactionResult Result = AutoRTFM::Transact([&]()
-				{
-					Unique = Controller.IsUnique();
-				});
+			AutoRTFM::Testing::Commit([&]()
+			{
+				Unique = Controller.IsUnique();
+			});
 
-			REQUIRE(AutoRTFM::ETransactionResult::Committed == Result);
 			REQUIRE(!Unique);
 		}
 	}
@@ -530,24 +497,21 @@ TEST_CASE("UECore.TIntrusiveReferenceController")
 
 		SECTION("With Abort")
 		{
-			AutoRTFM::ETransactionResult Result = AutoRTFM::Transact([&]()
-				{
-					Controller.ReleaseSharedReference();
-					AutoRTFM::AbortTransaction();
-				});
+			AutoRTFM::Testing::Abort([&]()
+			{
+				Controller.ReleaseSharedReference();
+				AutoRTFM::AbortTransaction();
+			});
 
-			REQUIRE(AutoRTFM::ETransactionResult::AbortedByRequest == Result);
 			REQUIRE(1 == Controller.GetSharedReferenceCount());
 		}
 
 		SECTION("With Commit")
 		{
-			AutoRTFM::ETransactionResult Result = AutoRTFM::Transact([&]()
-				{
-					Controller.ReleaseSharedReference();
-				});
-
-			REQUIRE(AutoRTFM::ETransactionResult::Committed == Result);
+			AutoRTFM::Testing::Commit([&]()
+			{
+				Controller.ReleaseSharedReference();
+			});
 		}
 	}
 
@@ -557,25 +521,22 @@ TEST_CASE("UECore.TIntrusiveReferenceController")
 
 		SECTION("With Abort")
 		{
-			AutoRTFM::ETransactionResult Result = AutoRTFM::Transact([&]()
-				{
-					Controller->ReleaseWeakReference();
-					AutoRTFM::AbortTransaction();
-				});
+			AutoRTFM::Testing::Abort([&]()
+			{
+				Controller->ReleaseWeakReference();
+				AutoRTFM::AbortTransaction();
+			});
 
-			REQUIRE(AutoRTFM::ETransactionResult::AbortedByRequest == Result);
 			REQUIRE(1 == Controller->WeakReferenceCount);
 			delete Controller;
 		}
 
 		SECTION("With Commit")
 		{
-			AutoRTFM::ETransactionResult Result = AutoRTFM::Transact([&]()
-				{
-					Controller->ReleaseWeakReference();
-				});
-
-			REQUIRE(AutoRTFM::ETransactionResult::Committed == Result);
+			AutoRTFM::Testing::Commit([&]()
+			{
+				Controller->ReleaseWeakReference();
+			});
 		}
 	}
 
@@ -585,24 +546,22 @@ TEST_CASE("UECore.TIntrusiveReferenceController")
 
 		SECTION("With Abort")
 		{
-			AutoRTFM::ETransactionResult Result = AutoRTFM::Transact([&]()
-				{
-					*Controller.GetObjectPtr() = 13;
-					AutoRTFM::AbortTransaction();
-				});
+			AutoRTFM::Testing::Abort([&]()
+			{
+				*Controller.GetObjectPtr() = 13;
+				AutoRTFM::AbortTransaction();
+			});
 
-			REQUIRE(AutoRTFM::ETransactionResult::AbortedByRequest == Result);
 			REQUIRE(42 == *Controller.GetObjectPtr());
 		}
 
 		SECTION("With Commit")
 		{
-			AutoRTFM::ETransactionResult Result = AutoRTFM::Transact([&]()
-				{
-					*Controller.GetObjectPtr() = 13;
-				});
+			AutoRTFM::Testing::Commit([&]()
+			{
+				*Controller.GetObjectPtr() = 13;
+			});
 
-			REQUIRE(AutoRTFM::ETransactionResult::Committed == Result);
 			REQUIRE(13 == *Controller.GetObjectPtr());
 		}
 	}
@@ -684,25 +643,21 @@ TEST_CASE("UECore.FTextCache")
 	{
 		SECTION("With Abort")
 		{
-			AutoRTFM::ETransactionResult Result = AutoRTFM::Transact([&]()
-				{
-					Cache.FindOrCache(TEXT("VALUE"), Key);
-					AutoRTFM::AbortTransaction();
-				});
-
-			REQUIRE(AutoRTFM::ETransactionResult::AbortedByRequest == Result);
+			AutoRTFM::Testing::Abort([&]()
+			{
+				Cache.FindOrCache(TEXT("VALUE"), Key);
+				AutoRTFM::AbortTransaction();
+			});
 
 			CheckCacheHealthy();
 		}
 
 		SECTION("With Commit")
 		{
-			AutoRTFM::ETransactionResult Result = AutoRTFM::Transact([&]()
-				{
-					Cache.FindOrCache(TEXT("VALUE"), Key);
-				});
-
-			REQUIRE(AutoRTFM::ETransactionResult::Committed == Result);
+			AutoRTFM::Testing::Commit([&]()
+			{
+				Cache.FindOrCache(TEXT("VALUE"), Key);
+			});
 
 			CheckCacheHealthy();
 		}
@@ -715,13 +670,11 @@ TEST_CASE("UECore.FTextCache")
 			// Add an entry to the cache before the transaction
 			Cache.FindOrCache(TEXT("VALUE"), Key);
 
-			AutoRTFM::ETransactionResult Result = AutoRTFM::Transact([&]()
-				{
-					Cache.FindOrCache(TEXT("REPLACEMENT"), Key);
-					AutoRTFM::AbortTransaction();
-				});
-
-			REQUIRE(AutoRTFM::ETransactionResult::AbortedByRequest == Result);
+			AutoRTFM::Testing::Abort([&]()
+			{
+				Cache.FindOrCache(TEXT("REPLACEMENT"), Key);
+				AutoRTFM::AbortTransaction();
+			});
 
 			CheckCacheHealthy();
 		}
@@ -731,12 +684,10 @@ TEST_CASE("UECore.FTextCache")
 			// Add an entry to the cache before the transaction
 			Cache.FindOrCache(TEXT("VALUE"), Key);
 
-			AutoRTFM::ETransactionResult Result = AutoRTFM::Transact([&]()
-				{
-					Cache.FindOrCache(TEXT("VALUE"), Key);
-				});
-
-			REQUIRE(AutoRTFM::ETransactionResult::Committed == Result);
+			AutoRTFM::Testing::Commit([&]()
+			{
+				Cache.FindOrCache(TEXT("VALUE"), Key);
+			});
 
 			CheckCacheHealthy();
 		}
@@ -749,13 +700,11 @@ TEST_CASE("UECore.FTextCache")
 			// Add an entry to the cache before the transaction
 			Cache.FindOrCache(TEXT("ORIGINAL"), Key);
 
-			AutoRTFM::ETransactionResult Result = AutoRTFM::Transact([&]()
-				{
-					Cache.FindOrCache(TEXT("REPLACEMENT"), Key);
-					AutoRTFM::AbortTransaction();
-				});
-
-			REQUIRE(AutoRTFM::ETransactionResult::AbortedByRequest == Result);
+			AutoRTFM::Testing::Abort([&]()
+			{
+				Cache.FindOrCache(TEXT("REPLACEMENT"), Key);
+				AutoRTFM::AbortTransaction();
+			});
 
 			CheckCacheHealthy();
 		}
@@ -765,12 +714,10 @@ TEST_CASE("UECore.FTextCache")
 			// Add an entry to the cache before the transaction
 			Cache.FindOrCache(TEXT("ORIGINAL"), Key);
 
-			AutoRTFM::ETransactionResult Result = AutoRTFM::Transact([&]()
-				{
-					Cache.FindOrCache(TEXT("REPLACEMENT"), Key);
-				});
-
-			REQUIRE(AutoRTFM::ETransactionResult::Committed == Result);
+			AutoRTFM::Testing::Commit([&]()
+			{
+				Cache.FindOrCache(TEXT("REPLACEMENT"), Key);
+			});
 
 			CheckCacheHealthy();
 		}
@@ -789,13 +736,11 @@ TEST_CASE("UECore.FTextCache")
 			// Add an entry to the cache before the transaction
 			Cache.FindOrCache(TEXT("VALUE"), Key);
 
-			AutoRTFM::ETransactionResult Result = AutoRTFM::Transact([&]()
-				{
-					Cache.RemoveCache(Key);
-					AutoRTFM::AbortTransaction();
-				});
-
-			REQUIRE(AutoRTFM::ETransactionResult::AbortedByRequest == Result);
+			AutoRTFM::Testing::Abort([&]()
+			{
+				Cache.RemoveCache(Key);
+				AutoRTFM::AbortTransaction();
+			});
 
 			CheckCacheHealthy();
 		}
@@ -805,12 +750,10 @@ TEST_CASE("UECore.FTextCache")
 			// Add an entry to the cache before the transaction
 			Cache.FindOrCache(TEXT("VALUE"), Key);
 
-			AutoRTFM::ETransactionResult Result = AutoRTFM::Transact([&]()
-				{
-					Cache.RemoveCache(Key);
-				});
-
-			REQUIRE(AutoRTFM::ETransactionResult::Committed == Result);
+			AutoRTFM::Testing::Commit([&]()
+			{
+				Cache.RemoveCache(Key);
+			});
 
 			CheckCacheHealthy();
 		}
@@ -823,27 +766,23 @@ TEST_CASE("UECore.FTextCache")
 		{
 			SECTION("With Abort")
 			{
-				AutoRTFM::ETransactionResult Result = AutoRTFM::Transact([&]()
-					{
-						Cache.FindOrCache(TEXT("VALUE"), Key);
-						AutoRTFM::Open([&]{ Cache.RemoveCache(Key); });
-						AutoRTFM::AbortTransaction();
-					});
-
-				REQUIRE(AutoRTFM::ETransactionResult::AbortedByRequest == Result);
+				AutoRTFM::Testing::Abort([&]()
+				{
+					Cache.FindOrCache(TEXT("VALUE"), Key);
+					AutoRTFM::Open([&]{ Cache.RemoveCache(Key); });
+					AutoRTFM::AbortTransaction();
+				});
 
 				CheckCacheHealthy();
 			}
 
 			SECTION("With Commit")
 			{
-				AutoRTFM::ETransactionResult Result = AutoRTFM::Transact([&]()
-					{
-						Cache.FindOrCache(TEXT("VALUE"), Key);
-						AutoRTFM::Open([&]{ Cache.RemoveCache(Key); });
-					});
-
-				REQUIRE(AutoRTFM::ETransactionResult::Committed == Result);
+				AutoRTFM::Testing::Commit([&]()
+				{
+					Cache.FindOrCache(TEXT("VALUE"), Key);
+					AutoRTFM::Open([&]{ Cache.RemoveCache(Key); });
+				});
 
 				CheckCacheHealthy();
 			}
@@ -862,24 +801,22 @@ TEST_CASE("UECore.FUObjectItem")
 		PROFILER_CHAR* const StatIDStringStorage = Item.StatIDStringStorage;
 
 		// If we abort then we won't change anything.
-		AutoRTFM::ETransactionResult Result = AutoRTFM::Transact([&]
-			{
-				Item.CreateStatID();
-				AutoRTFM::AbortTransaction();
-			});
+		AutoRTFM::Testing::Abort([&]
+		{
+			Item.CreateStatID();
+			AutoRTFM::AbortTransaction();
+		});
 
-		REQUIRE(AutoRTFM::ETransactionResult::AbortedByRequest == Result);
 		REQUIRE(StatIDStringStorage == Item.StatIDStringStorage);
 
 		// But also if we commit we likewise won't change anything because
 		// the string storage was already created before the transaction
 		// began.
-		Result = AutoRTFM::Transact([&]
-			{
-				Item.CreateStatID();
-			});
+		AutoRTFM::Testing::Commit([&]
+		{
+			Item.CreateStatID();
+		});
 
-		REQUIRE(AutoRTFM::ETransactionResult::Committed == Result);
 		REQUIRE(StatIDStringStorage == Item.StatIDStringStorage);
 	}
 
@@ -891,60 +828,54 @@ TEST_CASE("UECore.FUObjectItem")
 		REQUIRE(!Item.StatID.IsValidStat());
 
 		// If we abort then we won't change anything.
-		AutoRTFM::ETransactionResult Result = AutoRTFM::Transact([&]
-			{
-				Item.CreateStatID();
-				AutoRTFM::AbortTransaction();
-			});
+		AutoRTFM::Testing::Abort([&]
+		{
+			Item.CreateStatID();
+			AutoRTFM::AbortTransaction();
+		});
 
-		REQUIRE(AutoRTFM::ETransactionResult::AbortedByRequest == Result);
 		REQUIRE(nullptr == Item.StatIDStringStorage);
 		REQUIRE(!Item.StatID.IsValidStat());
 
 		// If we commit though we'll create the stat ID.
-		Result = AutoRTFM::Transact([&]
-			{
-				Item.CreateStatID();
-			});
+		AutoRTFM::Testing::Commit([&]
+		{
+			Item.CreateStatID();
+		});
 
-		REQUIRE(AutoRTFM::ETransactionResult::Committed == Result);
 		REQUIRE(nullptr != Item.StatIDStringStorage);
 		REQUIRE(Item.StatID.IsValidStat());
 	}
 
 	SECTION("CreateStatID On In-Transaction Object")
 	{
-		AutoRTFM::ETransactionResult Result = AutoRTFM::Transact([&]
-			{
-				FUObjectItem Item;
-				Item.Object = NewObject<UMyAutoRTFMTestObject>();
-				Item.CreateStatID();
+		AutoRTFM::Testing::Abort([&]
+		{
+			FUObjectItem Item;
+			Item.Object = NewObject<UMyAutoRTFMTestObject>();
+			Item.CreateStatID();
 
-				AutoRTFM::Open([&]
-					{
-						REQUIRE(nullptr != Item.StatIDStringStorage);
-						REQUIRE(Item.StatID.IsValidStat());
-					});
-
-				AutoRTFM::AbortTransaction();
-			});
-
-		REQUIRE(AutoRTFM::ETransactionResult::AbortedByRequest == Result);
+			AutoRTFM::Open([&]
+				{
+					REQUIRE(nullptr != Item.StatIDStringStorage);
+					REQUIRE(Item.StatID.IsValidStat());
+				});
 
-		Result = AutoRTFM::Transact([&]
-			{
-				FUObjectItem Item;
-				Item.Object = NewObject<UMyAutoRTFMTestObject>();
-				Item.CreateStatID();
+			AutoRTFM::AbortTransaction();
+		});
 
-				AutoRTFM::Open([&]
-					{
-						REQUIRE(nullptr != Item.StatIDStringStorage);
-						REQUIRE(Item.StatID.IsValidStat());
-					});
-			});
+		AutoRTFM::Testing::Commit([&]
+		{
+			FUObjectItem Item;
+			Item.Object = NewObject<UMyAutoRTFMTestObject>();
+			Item.CreateStatID();
 
-		REQUIRE(AutoRTFM::ETransactionResult::Committed == Result);
+			AutoRTFM::Open([&]
+				{
+					REQUIRE(nullptr != Item.StatIDStringStorage);
+					REQUIRE(Item.StatID.IsValidStat());
+				});
+		});
 	}
 
 	SECTION("CreateStatID In Closed Then Again In Open")
@@ -955,27 +886,26 @@ TEST_CASE("UECore.FUObjectItem")
 			REQUIRE(nullptr == Item.StatIDStringStorage);
 			REQUIRE(!Item.StatID.IsValidStat());
 
-			AutoRTFM::ETransactionResult Result = AutoRTFM::Transact([&]
-				{
-					Item.CreateStatID();
+			AutoRTFM::Testing::Abort([&]
+			{
+				Item.CreateStatID();
 
-					AutoRTFM::Open([&]
-						{
-							REQUIRE(nullptr != Item.StatIDStringStorage);
-							REQUIRE(Item.StatID.IsValidStat());
+				AutoRTFM::Open([&]
+					{
+						REQUIRE(nullptr != Item.StatIDStringStorage);
+						REQUIRE(Item.StatID.IsValidStat());
 
-							PROFILER_CHAR* const StatIDStringStorage = Item.StatIDStringStorage;
+						PROFILER_CHAR* const StatIDStringStorage = Item.StatIDStringStorage;
 
-							Item.CreateStatID();
+						Item.CreateStatID();
 
-							REQUIRE(StatIDStringStorage == Item.StatIDStringStorage);
-							REQUIRE(Item.StatID.IsValidStat());
-						});
+						REQUIRE(StatIDStringStorage == Item.StatIDStringStorage);
+						REQUIRE(Item.StatID.IsValidStat());
+					});
 
-					AutoRTFM::AbortTransaction();
-				});
+				AutoRTFM::AbortTransaction();
+			});
 
-			REQUIRE(AutoRTFM::ETransactionResult::AbortedByRequest == Result);
 			REQUIRE(nullptr == Item.StatIDStringStorage);
 			REQUIRE(!Item.StatID.IsValidStat());
 		}
@@ -986,25 +916,24 @@ TEST_CASE("UECore.FUObjectItem")
 			REQUIRE(nullptr == Item.StatIDStringStorage);
 			REQUIRE(!Item.StatID.IsValidStat());
 
-			AutoRTFM::ETransactionResult Result = AutoRTFM::Transact([&]
-				{
-					Item.CreateStatID();
+			AutoRTFM::Testing::Commit([&]
+			{
+				Item.CreateStatID();
 
-					AutoRTFM::Open([&]
-						{
-							REQUIRE(nullptr != Item.StatIDStringStorage);
-							REQUIRE(Item.StatID.IsValidStat());
+				AutoRTFM::Open([&]
+					{
+						REQUIRE(nullptr != Item.StatIDStringStorage);
+						REQUIRE(Item.StatID.IsValidStat());
 
-							PROFILER_CHAR* const StatIDStringStorage = Item.StatIDStringStorage;
+						PROFILER_CHAR* const StatIDStringStorage = Item.StatIDStringStorage;
 
-							Item.CreateStatID();
+						Item.CreateStatID();
 
-							REQUIRE(StatIDStringStorage == Item.StatIDStringStorage);
-							REQUIRE(Item.StatID.IsValidStat());
-						});
-				});
+						REQUIRE(StatIDStringStorage == Item.StatIDStringStorage);
+						REQUIRE(Item.StatID.IsValidStat());
+					});
+			});
 
-			REQUIRE(AutoRTFM::ETransactionResult::Committed == Result);
 			REQUIRE(nullptr != Item.StatIDStringStorage);
 			REQUIRE(Item.StatID.IsValidStat());
 		}
@@ -1017,99 +946,79 @@ TEST_CASE("UECore.TScopeLock_TransactionallySafeCriticalSection")
 	{
 		FTransactionallySafeCriticalSection CriticalSection;
 
-		AutoRTFM::ETransactionResult Result = AutoRTFM::Transact([&]
-			{
-				UE::TScopeLock Lock(CriticalSection);
-				AutoRTFM::AbortTransaction();
-			});
-
-		REQUIRE(AutoRTFM::ETransactionResult::AbortedByRequest == Result);
-
-		Result = AutoRTFM::Transact([&]
-			{
-				UE::TScopeLock Lock(CriticalSection);
-			});
+		AutoRTFM::Testing::Abort([&]
+		{
+			UE::TScopeLock Lock(CriticalSection);
+			AutoRTFM::AbortTransaction();
+		});
 
-		REQUIRE(AutoRTFM::ETransactionResult::Committed == Result);
+		AutoRTFM::Testing::Commit([&]
+		{
+			UE::TScopeLock Lock(CriticalSection);
+		});
 	}
 
 	SECTION("Inside Transaction")
 	{
-		AutoRTFM::ETransactionResult Result = AutoRTFM::Transact([&]
-			{
-				FTransactionallySafeCriticalSection CriticalSection;
-				UE::TScopeLock Lock(CriticalSection);
-				AutoRTFM::AbortTransaction();
-			});
-
-		REQUIRE(AutoRTFM::ETransactionResult::AbortedByRequest == Result);
-
-		Result = AutoRTFM::Transact([&]
-			{
-				FTransactionallySafeCriticalSection CriticalSection;
-				UE::TScopeLock Lock(CriticalSection);
-			});
+		AutoRTFM::Testing::Abort([&]
+		{
+			FTransactionallySafeCriticalSection CriticalSection;
+			UE::TScopeLock Lock(CriticalSection);
+			AutoRTFM::AbortTransaction();
+		});
 
-		REQUIRE(AutoRTFM::ETransactionResult::Committed == Result);
+		AutoRTFM::Testing::Commit([&]
+		{
+			FTransactionallySafeCriticalSection CriticalSection;
+			UE::TScopeLock Lock(CriticalSection);
+		});
 	}
 
 	SECTION("Inside Transaction Used In Nested Transaction")
 	{
-		AutoRTFM::ETransactionResult InnerResult = AutoRTFM::ETransactionResult::Committed;
+		AutoRTFM::Testing::Abort([&]
+		{
+			FTransactionallySafeCriticalSection CriticalSection;
 
-		AutoRTFM::ETransactionResult Result = AutoRTFM::Transact([&]
+			AutoRTFM::Testing::Abort([&]
 			{
-				FTransactionallySafeCriticalSection CriticalSection;
-
-				InnerResult = AutoRTFM::Transact([&]
-					{
-						UE::TScopeLock Lock(CriticalSection);
-						AutoRTFM::CascadingAbortTransaction();
-					});
+				UE::TScopeLock Lock(CriticalSection);
+				AutoRTFM::CascadingAbortTransaction();
 			});
+		});
 
-		REQUIRE(AutoRTFM::ETransactionResult::AbortedByCascade == Result);
+		AutoRTFM::Testing::Commit([&]
+		{
+			FTransactionallySafeCriticalSection CriticalSection;
 
-		Result = AutoRTFM::Transact([&]
+			AutoRTFM::Testing::Abort([&]
 			{
-				FTransactionallySafeCriticalSection CriticalSection;
-
-				InnerResult = AutoRTFM::Transact([&]
-					{
-						UE::TScopeLock Lock(CriticalSection);
-						AutoRTFM::AbortTransaction();
-					});
+				UE::TScopeLock Lock(CriticalSection);
+				AutoRTFM::AbortTransaction();
 			});
+		});
 
-		REQUIRE(AutoRTFM::ETransactionResult::AbortedByRequest == InnerResult);
-		REQUIRE(AutoRTFM::ETransactionResult::Committed == Result);
+		AutoRTFM::Testing::Abort([&]
+		{
+			FTransactionallySafeCriticalSection CriticalSection;
 
-		Result = AutoRTFM::Transact([&]
+			AutoRTFM::Testing::Commit([&]
 			{
-				FTransactionallySafeCriticalSection CriticalSection;
-
-				InnerResult = AutoRTFM::Transact([&]
-					{
-						UE::TScopeLock Lock(CriticalSection);
-					});
-
-				AutoRTFM::AbortTransaction();
+				UE::TScopeLock Lock(CriticalSection);
 			});
 
-		REQUIRE(AutoRTFM::ETransactionResult::AbortedByRequest == Result);
+			AutoRTFM::AbortTransaction();
+		});
 
-		Result = AutoRTFM::Transact([&]
-			{
-				FTransactionallySafeCriticalSection CriticalSection;
+		AutoRTFM::Testing::Commit([&]
+		{
+			FTransactionallySafeCriticalSection CriticalSection;
 
-				InnerResult = AutoRTFM::Transact([&]
-					{
-						UE::TScopeLock Lock(CriticalSection);
-					});
+			AutoRTFM::Testing::Commit([&]
+			{
+				UE::TScopeLock Lock(CriticalSection);
 			});
-
-		REQUIRE(AutoRTFM::ETransactionResult::Committed == InnerResult);
-		REQUIRE(AutoRTFM::ETransactionResult::Committed == Result);
+		});
 	}
 
 
@@ -1212,198 +1121,158 @@ TEST_CASE("UECore.TransactionallySafeRWScopeLock")
 	{
 		FTransactionallySafeRWLock RWLock;
 
-		AutoRTFM::ETransactionResult Result = AutoRTFM::Transact([&]
-			{
-				UE::TRWScopeLock Lock(RWLock, FRWScopeLockType::SLT_ReadOnly);
-				AutoRTFM::AbortTransaction();
-			});
-
-		REQUIRE(AutoRTFM::ETransactionResult::AbortedByRequest == Result);
-
-		Result = AutoRTFM::Transact([&]
-			{
-				UE::TRWScopeLock Lock(RWLock, FRWScopeLockType::SLT_ReadOnly);
-			});
+		AutoRTFM::Testing::Abort([&]
+		{
+			UE::TRWScopeLock Lock(RWLock, FRWScopeLockType::SLT_ReadOnly);
+			AutoRTFM::AbortTransaction();
+		});
 
-		REQUIRE(AutoRTFM::ETransactionResult::Committed == Result);
+		AutoRTFM::Testing::Commit([&]
+		{
+			UE::TRWScopeLock Lock(RWLock, FRWScopeLockType::SLT_ReadOnly);
+		});
 	}
 
 	SECTION("Inside Transaction With Read Lock")
 	{
-		AutoRTFM::ETransactionResult Result = AutoRTFM::Transact([&]
-			{
-				FTransactionallySafeRWLock RWLock;
-				UE::TRWScopeLock Lock(RWLock, FRWScopeLockType::SLT_ReadOnly);
-				AutoRTFM::AbortTransaction();
-			});
-
-		REQUIRE(AutoRTFM::ETransactionResult::AbortedByRequest == Result);
-
-		Result = AutoRTFM::Transact([&]
-			{
-				FTransactionallySafeRWLock RWLock;
-				UE::TRWScopeLock Lock(RWLock, FRWScopeLockType::SLT_ReadOnly);
-			});
+		AutoRTFM::Testing::Abort([&]
+		{
+			FTransactionallySafeRWLock RWLock;
+			UE::TRWScopeLock Lock(RWLock, FRWScopeLockType::SLT_ReadOnly);
+			AutoRTFM::AbortTransaction();
+		});
 
-		REQUIRE(AutoRTFM::ETransactionResult::Committed == Result);
+		AutoRTFM::Testing::Commit([&]
+		{
+			FTransactionallySafeRWLock RWLock;
+			UE::TRWScopeLock Lock(RWLock, FRWScopeLockType::SLT_ReadOnly);
+		});
 	}
 
 	SECTION("Inside Transaction Used In Nested Transaction With Read Lock")
 	{
-		AutoRTFM::ETransactionResult InnerResult = AutoRTFM::ETransactionResult::Committed;
+		AutoRTFM::Testing::Abort([&]
+		{
+			FTransactionallySafeRWLock RWLock;
 
-		AutoRTFM::ETransactionResult Result = AutoRTFM::Transact([&]
+			AutoRTFM::Testing::Abort([&]
 			{
-				FTransactionallySafeRWLock RWLock;
-
-				InnerResult = AutoRTFM::Transact([&]
-					{
-						UE::TRWScopeLock Lock(RWLock, FRWScopeLockType::SLT_ReadOnly);
-						AutoRTFM::CascadingAbortTransaction();
-					});
+				UE::TRWScopeLock Lock(RWLock, FRWScopeLockType::SLT_ReadOnly);
+				AutoRTFM::CascadingAbortTransaction();
 			});
+		});
 
-		REQUIRE(AutoRTFM::ETransactionResult::AbortedByCascade == Result);
+		AutoRTFM::Testing::Commit([&]
+		{
+			FTransactionallySafeRWLock RWLock;
 
-		Result = AutoRTFM::Transact([&]
+			AutoRTFM::Testing::Abort([&]
 			{
-				FTransactionallySafeRWLock RWLock;
-
-				InnerResult = AutoRTFM::Transact([&]
-					{
-						UE::TRWScopeLock Lock(RWLock, FRWScopeLockType::SLT_ReadOnly);
-						AutoRTFM::AbortTransaction();
-					});
+				UE::TRWScopeLock Lock(RWLock, FRWScopeLockType::SLT_ReadOnly);
+				AutoRTFM::AbortTransaction();
 			});
+		});
 
-		REQUIRE(AutoRTFM::ETransactionResult::AbortedByRequest == InnerResult);
-		REQUIRE(AutoRTFM::ETransactionResult::Committed == Result);
+		AutoRTFM::Testing::Abort([&]
+		{
+			FTransactionallySafeRWLock RWLock;
 
-		Result = AutoRTFM::Transact([&]
+			AutoRTFM::Testing::Commit([&]
 			{
-				FTransactionallySafeRWLock RWLock;
-
-				InnerResult = AutoRTFM::Transact([&]
-					{
-						UE::TRWScopeLock Lock(RWLock, FRWScopeLockType::SLT_ReadOnly);
-					});
-
-				AutoRTFM::AbortTransaction();
+				UE::TRWScopeLock Lock(RWLock, FRWScopeLockType::SLT_ReadOnly);
 			});
 
-		REQUIRE(AutoRTFM::ETransactionResult::AbortedByRequest == Result);
+			AutoRTFM::AbortTransaction();
+		});
 
-		Result = AutoRTFM::Transact([&]
-			{
-				FTransactionallySafeRWLock RWLock;
+		AutoRTFM::Testing::Commit([&]
+		{
+			FTransactionallySafeRWLock RWLock;
 
-				InnerResult = AutoRTFM::Transact([&]
-					{
-						UE::TRWScopeLock Lock(RWLock, FRWScopeLockType::SLT_ReadOnly);
-					});
+			AutoRTFM::Testing::Commit([&]
+			{
+				UE::TRWScopeLock Lock(RWLock, FRWScopeLockType::SLT_ReadOnly);
 			});
-
-		REQUIRE(AutoRTFM::ETransactionResult::Committed == InnerResult);
-		REQUIRE(AutoRTFM::ETransactionResult::Committed == Result);
+		});
 	}
 
 	SECTION("Outside Transaction With Write Lock")
 	{
 		FTransactionallySafeRWLock RWLock;
 
-		AutoRTFM::ETransactionResult Result = AutoRTFM::Transact([&]
-			{
-				UE::TRWScopeLock Lock(RWLock, FRWScopeLockType::SLT_Write);
-				AutoRTFM::AbortTransaction();
-			});
-
-		REQUIRE(AutoRTFM::ETransactionResult::AbortedByRequest == Result);
-
-		Result = AutoRTFM::Transact([&]
-			{
-				UE::TRWScopeLock Lock(RWLock, FRWScopeLockType::SLT_Write);
-			});
+		AutoRTFM::Testing::Abort([&]
+		{
+			UE::TRWScopeLock Lock(RWLock, FRWScopeLockType::SLT_Write);
+			AutoRTFM::AbortTransaction();
+		});
 
-		REQUIRE(AutoRTFM::ETransactionResult::Committed == Result);
+		AutoRTFM::Testing::Commit([&]
+		{
+			UE::TRWScopeLock Lock(RWLock, FRWScopeLockType::SLT_Write);
+		});
 	}
 
 	SECTION("Inside Transaction With Write Lock")
 	{
-		AutoRTFM::ETransactionResult Result = AutoRTFM::Transact([&]
-			{
-				FTransactionallySafeRWLock RWLock;
-				UE::TRWScopeLock Lock(RWLock, FRWScopeLockType::SLT_Write);
-				AutoRTFM::AbortTransaction();
-			});
-
-		REQUIRE(AutoRTFM::ETransactionResult::AbortedByRequest == Result);
-
-		Result = AutoRTFM::Transact([&]
-			{
-				FTransactionallySafeRWLock RWLock;
-				UE::TRWScopeLock Lock(RWLock, FRWScopeLockType::SLT_Write);
-			});
+		AutoRTFM::Testing::Abort([&]
+		{
+			FTransactionallySafeRWLock RWLock;
+			UE::TRWScopeLock Lock(RWLock, FRWScopeLockType::SLT_Write);
+			AutoRTFM::AbortTransaction();
+		});
 
-		REQUIRE(AutoRTFM::ETransactionResult::Committed == Result);
+		AutoRTFM::Testing::Commit([&]
+		{
+			FTransactionallySafeRWLock RWLock;
+			UE::TRWScopeLock Lock(RWLock, FRWScopeLockType::SLT_Write);
+		});
 	}
 
 	SECTION("Inside Transaction Used In Nested Transaction With Write Lock")
 	{
-		AutoRTFM::ETransactionResult InnerResult = AutoRTFM::ETransactionResult::Committed;
+		AutoRTFM::Testing::Abort([&]
+		{
+			FTransactionallySafeRWLock RWLock;
 
-		AutoRTFM::ETransactionResult Result = AutoRTFM::Transact([&]
+			AutoRTFM::Testing::Abort([&]
 			{
-				FTransactionallySafeRWLock RWLock;
-
-				InnerResult = AutoRTFM::Transact([&]
-					{
-						UE::TRWScopeLock Lock(RWLock, FRWScopeLockType::SLT_Write);
-						AutoRTFM::CascadingAbortTransaction();
-					});
+				UE::TRWScopeLock Lock(RWLock, FRWScopeLockType::SLT_Write);
+				AutoRTFM::CascadingAbortTransaction();
 			});
+		});
 
-		REQUIRE(AutoRTFM::ETransactionResult::AbortedByCascade == Result);
+		AutoRTFM::Testing::Commit([&]
+		{
+			FTransactionallySafeRWLock RWLock;
 
-		Result = AutoRTFM::Transact([&]
+			AutoRTFM::Testing::Abort([&]
 			{
-				FTransactionallySafeRWLock RWLock;
-
-				InnerResult = AutoRTFM::Transact([&]
-					{
-						UE::TRWScopeLock Lock(RWLock, FRWScopeLockType::SLT_Write);
-						AutoRTFM::AbortTransaction();
-					});
+				UE::TRWScopeLock Lock(RWLock, FRWScopeLockType::SLT_Write);
+				AutoRTFM::AbortTransaction();
 			});
+		});
 
-		REQUIRE(AutoRTFM::ETransactionResult::AbortedByRequest == InnerResult);
-		REQUIRE(AutoRTFM::ETransactionResult::Committed == Result);
+		AutoRTFM::Testing::Abort([&]
+		{
+			FTransactionallySafeRWLock RWLock;
 
-		Result = AutoRTFM::Transact([&]
+			AutoRTFM::Testing::Commit([&]
 			{
-				FTransactionallySafeRWLock RWLock;
-
-				InnerResult = AutoRTFM::Transact([&]
-					{
-						UE::TRWScopeLock Lock(RWLock, FRWScopeLockType::SLT_Write);
-					});
-
-				AutoRTFM::AbortTransaction();
+				UE::TRWScopeLock Lock(RWLock, FRWScopeLockType::SLT_Write);
 			});
 
-		REQUIRE(AutoRTFM::ETransactionResult::AbortedByRequest == Result);
+			AutoRTFM::AbortTransaction();
+		});
 
-		Result = AutoRTFM::Transact([&]
-			{
-				FTransactionallySafeRWLock RWLock;
+		AutoRTFM::Testing::Commit([&]
+		{
+			FTransactionallySafeRWLock RWLock;
 
-				InnerResult = AutoRTFM::Transact([&]
-					{
-						UE::TRWScopeLock Lock(RWLock, FRWScopeLockType::SLT_Write);
-					});
+			AutoRTFM::Testing::Commit([&]
+			{
+				UE::TRWScopeLock Lock(RWLock, FRWScopeLockType::SLT_Write);
 			});
-
-		REQUIRE(AutoRTFM::ETransactionResult::Committed == InnerResult);
-		REQUIRE(AutoRTFM::ETransactionResult::Committed == Result);
+		});
 	}
 }
 
@@ -1413,99 +1282,79 @@ TEST_CASE("UECore.TransactionallySafeReadScopeLock")
 	{
 		FTransactionallySafeRWLock RWLock;
 
-		AutoRTFM::ETransactionResult Result = AutoRTFM::Transact([&]
-			{
-				UE::TReadScopeLock Lock(RWLock);
-				AutoRTFM::AbortTransaction();
-			});
-
-		REQUIRE(AutoRTFM::ETransactionResult::AbortedByRequest == Result);
-
-		Result = AutoRTFM::Transact([&]
-			{
-				UE::TReadScopeLock Lock(RWLock);
-			});
+		AutoRTFM::Testing::Abort([&]
+		{
+			UE::TReadScopeLock Lock(RWLock);
+			AutoRTFM::AbortTransaction();
+		});
 
-		REQUIRE(AutoRTFM::ETransactionResult::Committed == Result);
+		AutoRTFM::Testing::Commit([&]
+		{
+			UE::TReadScopeLock Lock(RWLock);
+		});
 	}
 
 	SECTION("Inside Transaction")
 	{
-		AutoRTFM::ETransactionResult Result = AutoRTFM::Transact([&]
-			{
-				FTransactionallySafeRWLock RWLock;
-				UE::TReadScopeLock Lock(RWLock);
-				AutoRTFM::AbortTransaction();
-			});
-
-		REQUIRE(AutoRTFM::ETransactionResult::AbortedByRequest == Result);
-
-		Result = AutoRTFM::Transact([&]
-			{
-				FTransactionallySafeRWLock RWLock;
-				UE::TReadScopeLock Lock(RWLock);
-			});
+		AutoRTFM::Testing::Abort([&]
+		{
+			FTransactionallySafeRWLock RWLock;
+			UE::TReadScopeLock Lock(RWLock);
+			AutoRTFM::AbortTransaction();
+		});
 
-		REQUIRE(AutoRTFM::ETransactionResult::Committed == Result);
+		AutoRTFM::Testing::Commit([&]
+		{
+			FTransactionallySafeRWLock RWLock;
+			UE::TReadScopeLock Lock(RWLock);
+		});
 	}
 
 	SECTION("Inside Transaction Used In Nested Transaction")
 	{
-		AutoRTFM::ETransactionResult InnerResult = AutoRTFM::ETransactionResult::Committed;
+		AutoRTFM::Testing::Abort([&]
+		{
+			FTransactionallySafeRWLock RWLock;
 
-		AutoRTFM::ETransactionResult Result = AutoRTFM::Transact([&]
+			AutoRTFM::Testing::Abort([&]
 			{
-				FTransactionallySafeRWLock RWLock;
-
-				InnerResult = AutoRTFM::Transact([&]
-					{
-						UE::TReadScopeLock Lock(RWLock);
-						AutoRTFM::CascadingAbortTransaction();
-					});
+				UE::TReadScopeLock Lock(RWLock);
+				AutoRTFM::CascadingAbortTransaction();
 			});
+		});
 
-		REQUIRE(AutoRTFM::ETransactionResult::AbortedByCascade == Result);
+		AutoRTFM::Testing::Commit([&]
+		{
+			FTransactionallySafeRWLock RWLock;
 
-		Result = AutoRTFM::Transact([&]
+			AutoRTFM::Testing::Abort([&]
 			{
-				FTransactionallySafeRWLock RWLock;
-
-				InnerResult = AutoRTFM::Transact([&]
-					{
-						UE::TReadScopeLock Lock(RWLock);
-						AutoRTFM::AbortTransaction();
-					});
+				UE::TReadScopeLock Lock(RWLock);
+				AutoRTFM::AbortTransaction();
 			});
+		});
 
-		REQUIRE(AutoRTFM::ETransactionResult::AbortedByRequest == InnerResult);
-		REQUIRE(AutoRTFM::ETransactionResult::Committed == Result);
+		AutoRTFM::Testing::Abort([&]
+		{
+			FTransactionallySafeRWLock RWLock;
 
-		Result = AutoRTFM::Transact([&]
+			AutoRTFM::Testing::Commit([&]
 			{
-				FTransactionallySafeRWLock RWLock;
-
-				InnerResult = AutoRTFM::Transact([&]
-					{
-						UE::TReadScopeLock Lock(RWLock);
-					});
-
-				AutoRTFM::AbortTransaction();
+				UE::TReadScopeLock Lock(RWLock);
 			});
 
-		REQUIRE(AutoRTFM::ETransactionResult::AbortedByRequest == Result);
+			AutoRTFM::AbortTransaction();
+		});
 
-		Result = AutoRTFM::Transact([&]
-			{
-				FTransactionallySafeRWLock RWLock;
+		AutoRTFM::Testing::Commit([&]
+		{
+			FTransactionallySafeRWLock RWLock;
 
-				InnerResult = AutoRTFM::Transact([&]
-					{
-						UE::TReadScopeLock Lock(RWLock);
-					});
+			AutoRTFM::Testing::Commit([&]
+			{
+				UE::TReadScopeLock Lock(RWLock);
 			});
-
-		REQUIRE(AutoRTFM::ETransactionResult::Committed == InnerResult);
-		REQUIRE(AutoRTFM::ETransactionResult::Committed == Result);
+		});
 	}
 
 	SECTION("TReadScopeLock, destruct, memzero, reconstruct")
@@ -1547,116 +1396,94 @@ TEST_CASE("UECore.TransactionallySafeRWLockTryWriteLock")
 	{
 		FTransactionallySafeRWLock RWLock;
 
-		AutoRTFM::ETransactionResult Result = AutoRTFM::Transact([&]
-			{
-				REQUIRE(RWLock.TryWriteLock());
-				AutoRTFM::AbortTransaction();
-			});
-
-		REQUIRE(AutoRTFM::ETransactionResult::AbortedByRequest == Result);
-
-		Result = AutoRTFM::Transact([&]
-			{
-				REQUIRE(RWLock.TryWriteLock());
-				RWLock.WriteUnlock();
-			});
+		AutoRTFM::Testing::Abort([&]
+		{
+			REQUIRE(RWLock.TryWriteLock());
+			AutoRTFM::AbortTransaction();
+		});
 
-		REQUIRE(AutoRTFM::ETransactionResult::Committed == Result);
+		AutoRTFM::Testing::Commit([&]
+		{
+			REQUIRE(RWLock.TryWriteLock());
+			RWLock.WriteUnlock();
+		});
 	}
 
 	SECTION("Inside Transaction")
 	{
-		AutoRTFM::ETransactionResult Result = AutoRTFM::Transact([&]
-			{
-				FTransactionallySafeRWLock RWLock;
-				REQUIRE(RWLock.TryWriteLock());
-				AutoRTFM::AbortTransaction();
-			});
-
-		REQUIRE(AutoRTFM::ETransactionResult::AbortedByRequest == Result);
-
-		Result = AutoRTFM::Transact([&]
-			{
-				FTransactionallySafeRWLock RWLock;
-				REQUIRE(RWLock.TryWriteLock());
-				RWLock.WriteUnlock();
-			});
+		AutoRTFM::Testing::Abort([&]
+		{
+			FTransactionallySafeRWLock RWLock;
+			REQUIRE(RWLock.TryWriteLock());
+			AutoRTFM::AbortTransaction();
+		});
 
-		REQUIRE(AutoRTFM::ETransactionResult::Committed == Result);
+		AutoRTFM::Testing::Commit([&]
+		{
+			FTransactionallySafeRWLock RWLock;
+			REQUIRE(RWLock.TryWriteLock());
+			RWLock.WriteUnlock();
+		});
 	}
 
 	SECTION("Inside Transaction Used In Nested Transaction")
 	{
-		AutoRTFM::ETransactionResult InnerResult = AutoRTFM::ETransactionResult::Committed;
+		AutoRTFM::Testing::Abort([&]
+		{
+			FTransactionallySafeRWLock RWLock;
 
-		AutoRTFM::ETransactionResult Result = AutoRTFM::Transact([&]
+			AutoRTFM::Testing::Abort([&]
 			{
-				FTransactionallySafeRWLock RWLock;
-
-				InnerResult = AutoRTFM::Transact([&]
-					{
-						REQUIRE(RWLock.TryWriteLock());
-						AutoRTFM::CascadingAbortTransaction();
-					});
+				REQUIRE(RWLock.TryWriteLock());
+				AutoRTFM::CascadingAbortTransaction();
 			});
+		});
 
-		REQUIRE(AutoRTFM::ETransactionResult::AbortedByCascade == Result);
+		AutoRTFM::Testing::Commit([&]
+		{
+			FTransactionallySafeRWLock RWLock;
 
-		Result = AutoRTFM::Transact([&]
+			AutoRTFM::Testing::Abort([&]
 			{
-				FTransactionallySafeRWLock RWLock;
-
-				InnerResult = AutoRTFM::Transact([&]
-					{
-						REQUIRE(RWLock.TryWriteLock());
-						AutoRTFM::AbortTransaction();
-					});
+				REQUIRE(RWLock.TryWriteLock());
+				AutoRTFM::AbortTransaction();
 			});
+		});
 
-		REQUIRE(AutoRTFM::ETransactionResult::AbortedByRequest == InnerResult);
-		REQUIRE(AutoRTFM::ETransactionResult::Committed == Result);
+		AutoRTFM::Testing::Abort([&]
+		{
+			FTransactionallySafeRWLock RWLock;
 
-		Result = AutoRTFM::Transact([&]
+			AutoRTFM::Testing::Commit([&]
 			{
-				FTransactionallySafeRWLock RWLock;
-
-				InnerResult = AutoRTFM::Transact([&]
-					{
-						REQUIRE(RWLock.TryWriteLock());
-					});
-
-				AutoRTFM::AbortTransaction();
+				REQUIRE(RWLock.TryWriteLock());
 			});
 
-		REQUIRE(AutoRTFM::ETransactionResult::AbortedByRequest == Result);
-
-		Result = AutoRTFM::Transact([&]
-			{
-				FTransactionallySafeRWLock RWLock;
+			AutoRTFM::AbortTransaction();
+		});
 
-				InnerResult = AutoRTFM::Transact([&]
-					{
-						REQUIRE(RWLock.TryWriteLock());
-					});
+		AutoRTFM::Testing::Commit([&]
+		{
+			FTransactionallySafeRWLock RWLock;
 
-				RWLock.WriteUnlock();
+			AutoRTFM::Testing::Commit([&]
+			{
+				REQUIRE(RWLock.TryWriteLock());
 			});
 
-		REQUIRE(AutoRTFM::ETransactionResult::Committed == InnerResult);
-		REQUIRE(AutoRTFM::ETransactionResult::Committed == Result);
+			RWLock.WriteUnlock();
+		});
 	}
 
 	SECTION("Multiple TryWriteLocks In Transaction")
 	{
-		AutoRTFM::ETransactionResult Result = AutoRTFM::Transact([&]
-			{
-				FTransactionallySafeRWLock RWLock;
-				REQUIRE(RWLock.TryWriteLock());
-				REQUIRE(RWLock.TryWriteLock());
-				RWLock.WriteUnlock();
-			});
-
-		REQUIRE(AutoRTFM::ETransactionResult::Committed == Result);
+		AutoRTFM::Testing::Commit([&]
+		{
+			FTransactionallySafeRWLock RWLock;
+			REQUIRE(RWLock.TryWriteLock());
+			REQUIRE(RWLock.TryWriteLock());
+			RWLock.WriteUnlock();
+		});
 	}
 }
 
@@ -1667,99 +1494,79 @@ TEST_CASE("UECore.TransactionallySafeWriteScopeLock")
 	{
 		FTransactionallySafeRWLock RWLock;
 
-		AutoRTFM::ETransactionResult Result = AutoRTFM::Transact([&]
-			{
-				UE::TWriteScopeLock Lock(RWLock);
-				AutoRTFM::AbortTransaction();
-			});
-
-		REQUIRE(AutoRTFM::ETransactionResult::AbortedByRequest == Result);
-
-		Result = AutoRTFM::Transact([&]
-			{
-				UE::TWriteScopeLock Lock(RWLock);
-			});
-
-		REQUIRE(AutoRTFM::ETransactionResult::Committed == Result);
-	}
-
-	SECTION("Inside Transaction")
-	{
-		AutoRTFM::ETransactionResult Result = AutoRTFM::Transact([&]
-			{
-				FTransactionallySafeRWLock RWLock;
-				UE::TWriteScopeLock Lock(RWLock);
-				AutoRTFM::AbortTransaction();
-			});
+		AutoRTFM::Testing::Abort([&]
+		{
+			UE::TWriteScopeLock Lock(RWLock);
+			AutoRTFM::AbortTransaction();
+		});
 
-		REQUIRE(AutoRTFM::ETransactionResult::AbortedByRequest == Result);
+		AutoRTFM::Testing::Commit([&]
+		{
+			UE::TWriteScopeLock Lock(RWLock);
+		});
+	}
 
-		Result = AutoRTFM::Transact([&]
-			{
-				FTransactionallySafeRWLock RWLock;
-				UE::TWriteScopeLock Lock(RWLock);
-			});
+	SECTION("Inside Transaction")
+	{
+		AutoRTFM::Testing::Abort([&]
+		{
+			FTransactionallySafeRWLock RWLock;
+			UE::TWriteScopeLock Lock(RWLock);
+			AutoRTFM::AbortTransaction();
+		});
 
-		REQUIRE(AutoRTFM::ETransactionResult::Committed == Result);
+		AutoRTFM::Testing::Commit([&]
+		{
+			FTransactionallySafeRWLock RWLock;
+			UE::TWriteScopeLock Lock(RWLock);
+		});
 	}
 
 	SECTION("Inside Transaction Used In Nested Transaction")
 	{
-		AutoRTFM::ETransactionResult InnerResult = AutoRTFM::ETransactionResult::Committed;
+		AutoRTFM::Testing::Abort([&]
+		{
+			FTransactionallySafeRWLock RWLock;
 
-		AutoRTFM::ETransactionResult Result = AutoRTFM::Transact([&]
+			AutoRTFM::Testing::Abort([&]
 			{
-				FTransactionallySafeRWLock RWLock;
-
-				InnerResult = AutoRTFM::Transact([&]
-					{
-						UE::TWriteScopeLock Lock(RWLock);
-						AutoRTFM::CascadingAbortTransaction();
-					});
+				UE::TWriteScopeLock Lock(RWLock);
+				AutoRTFM::CascadingAbortTransaction();
 			});
+		});
 
-		REQUIRE(AutoRTFM::ETransactionResult::AbortedByCascade == Result);
+		AutoRTFM::Testing::Commit([&]
+		{
+			FTransactionallySafeRWLock RWLock;
 
-		Result = AutoRTFM::Transact([&]
+			AutoRTFM::Testing::Abort([&]
 			{
-				FTransactionallySafeRWLock RWLock;
-
-				InnerResult = AutoRTFM::Transact([&]
-					{
-						UE::TWriteScopeLock Lock(RWLock);
-						AutoRTFM::AbortTransaction();
-					});
+				UE::TWriteScopeLock Lock(RWLock);
+				AutoRTFM::AbortTransaction();
 			});
+		});
 
-		REQUIRE(AutoRTFM::ETransactionResult::AbortedByRequest == InnerResult);
-		REQUIRE(AutoRTFM::ETransactionResult::Committed == Result);
+		AutoRTFM::Testing::Abort([&]
+		{
+			FTransactionallySafeRWLock RWLock;
 
-		Result = AutoRTFM::Transact([&]
+			AutoRTFM::Testing::Commit([&]
 			{
-				FTransactionallySafeRWLock RWLock;
-
-				InnerResult = AutoRTFM::Transact([&]
-					{
-						UE::TWriteScopeLock Lock(RWLock);
-					});
-
-				AutoRTFM::AbortTransaction();
+				UE::TWriteScopeLock Lock(RWLock);
 			});
 
-		REQUIRE(AutoRTFM::ETransactionResult::AbortedByRequest == Result);
+			AutoRTFM::AbortTransaction();
+		});
+
+		AutoRTFM::Testing::Commit([&]
+		{
+			FTransactionallySafeRWLock RWLock;
 
-		Result = AutoRTFM::Transact([&]
+			AutoRTFM::Testing::Commit([&]
 			{
-				FTransactionallySafeRWLock RWLock;
-
-				InnerResult = AutoRTFM::Transact([&]
-					{
-						UE::TWriteScopeLock Lock(RWLock);
-					});
+				UE::TWriteScopeLock Lock(RWLock);
 			});
-
-		REQUIRE(AutoRTFM::ETransactionResult::Committed == InnerResult);
-		REQUIRE(AutoRTFM::ETransactionResult::Committed == Result);
+		});
 	}
 
 	SECTION("TWriteScopeLock, destruct, memzero, reconstruct")
@@ -1801,20 +1608,16 @@ TEST_CASE("UECore.TransactionallySafeMutex")
 	{
 		UE::FTransactionallySafeMutex Mutex;
 
-		AutoRTFM::ETransactionResult Result = AutoRTFM::Transact([&]
-			{
-				UE::TScopeLock Lock(Mutex);
-				AutoRTFM::AbortTransaction();
-			});
-
-		REQUIRE(AutoRTFM::ETransactionResult::AbortedByRequest == Result);
-
-		Result = AutoRTFM::Transact([&]
-			{
-				UE::TScopeLock Lock(Mutex);
-			});
+		AutoRTFM::Testing::Abort([&]
+		{
+			UE::TScopeLock Lock(Mutex);
+			AutoRTFM::AbortTransaction();
+		});
 
-		REQUIRE(AutoRTFM::ETransactionResult::Committed == Result);
+		AutoRTFM::Testing::Commit([&]
+		{
+			UE::TScopeLock Lock(Mutex);
+		});
 	}
 
 	SECTION("IsLocked Outside Transaction")
@@ -1830,22 +1633,18 @@ TEST_CASE("UECore.TransactionallySafeMutex")
 
 	SECTION("Inside Transaction")
 	{
-		AutoRTFM::ETransactionResult Result = AutoRTFM::Transact([&]
-			{
-				UE::FTransactionallySafeMutex Mutex;
-				UE::TScopeLock Lock(Mutex);
-				AutoRTFM::AbortTransaction();
-			});
-
-		REQUIRE(AutoRTFM::ETransactionResult::AbortedByRequest == Result);
-
-		Result = AutoRTFM::Transact([&]
-			{
-				UE::FTransactionallySafeMutex Mutex;
-				UE::TScopeLock Lock(Mutex);
-			});
+		AutoRTFM::Testing::Abort([&]
+		{
+			UE::FTransactionallySafeMutex Mutex;
+			UE::TScopeLock Lock(Mutex);
+			AutoRTFM::AbortTransaction();
+		});
 
-		REQUIRE(AutoRTFM::ETransactionResult::Committed == Result);
+		AutoRTFM::Testing::Commit([&]
+		{
+			UE::FTransactionallySafeMutex Mutex;
+			UE::TScopeLock Lock(Mutex);
+		});
 	}
 
 	SECTION("IsLocked Inside Transaction")
@@ -1867,61 +1666,49 @@ TEST_CASE("UECore.TransactionallySafeMutex")
 
 	SECTION("Inside Transaction Used In Nested Transaction")
 	{
-		AutoRTFM::ETransactionResult InnerResult = AutoRTFM::ETransactionResult::Committed;
+		AutoRTFM::Testing::Abort([&]
+		{
+			UE::FTransactionallySafeMutex Mutex;
 
-		AutoRTFM::ETransactionResult Result = AutoRTFM::Transact([&]
+			AutoRTFM::Testing::Abort([&]
 			{
-				UE::FTransactionallySafeMutex Mutex;
-
-				InnerResult = AutoRTFM::Transact([&]
-					{
-						UE::TScopeLock Lock(Mutex);
-						AutoRTFM::CascadingAbortTransaction();
-					});
+				UE::TScopeLock Lock(Mutex);
+				AutoRTFM::CascadingAbortTransaction();
 			});
+		});
 
-		REQUIRE(AutoRTFM::ETransactionResult::AbortedByCascade == Result);
+		AutoRTFM::Testing::Commit([&]
+		{
+			UE::FTransactionallySafeMutex Mutex;
 
-		Result = AutoRTFM::Transact([&]
+			AutoRTFM::Testing::Abort([&]
 			{
-				UE::FTransactionallySafeMutex Mutex;
-
-				InnerResult = AutoRTFM::Transact([&]
-					{
-						UE::TScopeLock Lock(Mutex);
-						AutoRTFM::AbortTransaction();
-					});
+				UE::TScopeLock Lock(Mutex);
+				AutoRTFM::AbortTransaction();
 			});
+		});
 
-		REQUIRE(AutoRTFM::ETransactionResult::AbortedByRequest == InnerResult);
-		REQUIRE(AutoRTFM::ETransactionResult::Committed == Result);
+		AutoRTFM::Testing::Abort([&]
+		{
+			UE::FTransactionallySafeMutex Mutex;
 
-		Result = AutoRTFM::Transact([&]
+			AutoRTFM::Testing::Commit([&]
 			{
-				UE::FTransactionallySafeMutex Mutex;
-
-				InnerResult = AutoRTFM::Transact([&]
-					{
-						UE::TScopeLock Lock(Mutex);
-					});
-
-				AutoRTFM::AbortTransaction();
+				UE::TScopeLock Lock(Mutex);
 			});
 
-		REQUIRE(AutoRTFM::ETransactionResult::AbortedByRequest == Result);
+			AutoRTFM::AbortTransaction();
+		});
 
-		Result = AutoRTFM::Transact([&]
-			{
-				UE::FTransactionallySafeMutex Mutex;
+		AutoRTFM::Testing::Commit([&]
+		{
+			UE::FTransactionallySafeMutex Mutex;
 
-				InnerResult = AutoRTFM::Transact([&]
-					{
-						UE::TScopeLock Lock(Mutex);
-					});
+			AutoRTFM::Testing::Commit([&]
+			{
+				UE::TScopeLock Lock(Mutex);
 			});
-
-		REQUIRE(AutoRTFM::ETransactionResult::Committed == InnerResult);
-		REQUIRE(AutoRTFM::ETransactionResult::Committed == Result);
+		});
 	}
 
 	SECTION("In Static Local Initializer")
@@ -2659,21 +2446,19 @@ TEST_CASE("UECore.FTextFormatPatternDefinition")
 
 	REQUIRE(!Ptr.IsValid());
 
-	AutoRTFM::ETransactionResult Result = AutoRTFM::Transact([&]
-		{
-			Ptr = FTextFormatPatternDefinition::GetDefault().ToSharedPtr();
-			AutoRTFM::AbortTransaction();
-		});
+	AutoRTFM::Testing::Abort([&]
+	{
+		Ptr = FTextFormatPatternDefinition::GetDefault().ToSharedPtr();
+		AutoRTFM::AbortTransaction();
+	});
 
-	REQUIRE(AutoRTFM::ETransactionResult::AbortedByRequest == Result);
 	REQUIRE(!Ptr.IsValid());
 
-	Result = AutoRTFM::Transact([&]
-		{
-			Ptr = FTextFormatPatternDefinition::GetDefault().ToSharedPtr();
-		});
+	AutoRTFM::Testing::Commit([&]
+	{
+		Ptr = FTextFormatPatternDefinition::GetDefault().ToSharedPtr();
+	});
 
-	REQUIRE(AutoRTFM::ETransactionResult::Committed == Result);
 	REQUIRE(Ptr.IsValid());
 }
 
@@ -2683,10 +2468,10 @@ TEST_CASE("UECore.FString")
 	{
 		FString String;
 
-		AutoRTFM::Commit([&]
-			{
-				String = FString::Printf(TEXT("Foo '%s' Bar"), TEXT("Stuff"));
-			});
+		AutoRTFM::Testing::Commit([&]
+		{
+			String = FString::Printf(TEXT("Foo '%s' Bar"), TEXT("Stuff"));
+		});
 
 		REQUIRE(String == "Foo 'Stuff' BAR");
 	}
@@ -2697,13 +2482,13 @@ TEST_CASE("UECore.FString")
 		{
 			FString String;
 
-			AutoRTFM::Commit([&]
+			AutoRTFM::Testing::Commit([&]
+			{
+				String = AutoRTFM::Open([&]
 				{
-					String = AutoRTFM::Open([&]
-						{
-							return TEXT("WOW");
-						});
+					return TEXT("WOW");
 				});
+			});
 
 			REQUIRE(String == "WOW");
 		}
@@ -2713,13 +2498,13 @@ TEST_CASE("UECore.FString")
 			FString Other = TEXT("WOW");
 			FString String;
 
-			AutoRTFM::Commit([&]
+			AutoRTFM::Testing::Commit([&]
+			{
+				String = AutoRTFM::Open([&]
 				{
-					String = AutoRTFM::Open([&]
-						{
-							return Other;
-						});
+					return Other;
 				});
+			});
 
 			REQUIRE(Other == "WOW");
 			REQUIRE(String == "WOW");
@@ -2733,15 +2518,15 @@ TEST_CASE("UECore.TQueue")
 	{
 		SECTION("Constructor")
 		{
-			AutoRTFM::Commit([&]
-				{
-					TQueue<int, EQueueMode::SingleThreaded> Queue;
+			AutoRTFM::Testing::Commit([&]
+			{
+				TQueue<int, EQueueMode::SingleThreaded> Queue;
 
-					AutoRTFM::Open([&]
-						{
-							REQUIRE(nullptr == Queue.Peek());
-						});
+				AutoRTFM::Open([&]
+				{
+					REQUIRE(nullptr == Queue.Peek());
 				});
+			});
 		}
 
 		SECTION("Dequeue")
@@ -2753,23 +2538,21 @@ TEST_CASE("UECore.TQueue")
 			int Value = 0;
 			bool bSucceeded = false;
 
-			AutoRTFM::ETransactionResult Result = AutoRTFM::Transact([&]
-				{
-					bSucceeded = Queue.Dequeue(Value);
-					AutoRTFM::AbortTransaction();
-				});
+			AutoRTFM::Testing::Abort([&]
+			{
+				bSucceeded = Queue.Dequeue(Value);
+				AutoRTFM::AbortTransaction();
+			});
 
-			REQUIRE(AutoRTFM::ETransactionResult::AbortedByRequest == Result);
 			REQUIRE(!bSucceeded);
 			REQUIRE(0 == Value);
 			REQUIRE(42 == *Queue.Peek());
 
-			Result = AutoRTFM::Transact([&]
-				{
-					bSucceeded = Queue.Dequeue(Value);
-				});
+			AutoRTFM::Testing::Commit([&]
+			{
+				bSucceeded = Queue.Dequeue(Value);
+			});
 
-			REQUIRE(AutoRTFM::ETransactionResult::Committed == Result);
 			REQUIRE(bSucceeded);
 			REQUIRE(42 == Value);
 			REQUIRE(Queue.IsEmpty());
@@ -2781,22 +2564,20 @@ TEST_CASE("UECore.TQueue")
 			REQUIRE(Queue.Enqueue(42));
 			REQUIRE(!Queue.IsEmpty());
 
-			AutoRTFM::ETransactionResult Result = AutoRTFM::Transact([&]
-				{
-					Queue.Empty();
-					AutoRTFM::Open([&] { REQUIRE(Queue.IsEmpty()); });
-					AutoRTFM::AbortTransaction();
-				});
+			AutoRTFM::Testing::Abort([&]
+			{
+				Queue.Empty();
+				AutoRTFM::Open([&] { REQUIRE(Queue.IsEmpty()); });
+				AutoRTFM::AbortTransaction();
+			});
 
-			REQUIRE(AutoRTFM::ETransactionResult::AbortedByRequest == Result);
 			REQUIRE(42 == *Queue.Peek());
 
-			Result = AutoRTFM::Transact([&]
-				{
-					Queue.Empty();
-				});
+			AutoRTFM::Testing::Commit([&]
+			{
+				Queue.Empty();
+			});
 
-			REQUIRE(AutoRTFM::ETransactionResult::Committed == Result);
 			REQUIRE(Queue.IsEmpty());
 		}
 
@@ -2806,22 +2587,20 @@ TEST_CASE("UECore.TQueue")
 
 			bool bSucceeded = false;
 
-			AutoRTFM::ETransactionResult Result = AutoRTFM::Transact([&]
-				{
-					bSucceeded = Queue.Enqueue(42);
-					AutoRTFM::AbortTransaction();
-				});
+			AutoRTFM::Testing::Abort([&]
+			{
+				bSucceeded = Queue.Enqueue(42);
+				AutoRTFM::AbortTransaction();
+			});
 
-			REQUIRE(AutoRTFM::ETransactionResult::AbortedByRequest == Result);
 			REQUIRE(Queue.IsEmpty());
 			REQUIRE(!bSucceeded);
 
-			Result = AutoRTFM::Transact([&]
-				{
-					bSucceeded = Queue.Enqueue(42);
-				});
+			AutoRTFM::Testing::Commit([&]
+			{
+				bSucceeded = Queue.Enqueue(42);
+			});
 
-			REQUIRE(AutoRTFM::ETransactionResult::Committed == Result);
 			REQUIRE(42 == *Queue.Peek());
 			REQUIRE(bSucceeded);
 		}
@@ -2833,41 +2612,37 @@ TEST_CASE("UECore.TQueue")
 
 			bool bIsEmpty = false;
 
-			AutoRTFM::ETransactionResult Result = AutoRTFM::Transact([&]
-				{
-					bIsEmpty = Queue.IsEmpty();
-					AutoRTFM::AbortTransaction();
-				});
+			AutoRTFM::Testing::Abort([&]
+			{
+				bIsEmpty = Queue.IsEmpty();
+				AutoRTFM::AbortTransaction();
+			});
 
-			REQUIRE(AutoRTFM::ETransactionResult::AbortedByRequest == Result);
 			REQUIRE(!bIsEmpty);
 			
-			Result = AutoRTFM::Transact([&]
-				{
-					bIsEmpty = Queue.IsEmpty();
-				});
+			AutoRTFM::Testing::Commit([&]
+			{
+				bIsEmpty = Queue.IsEmpty();
+			});
 
-			REQUIRE(AutoRTFM::ETransactionResult::Committed == Result);
 			REQUIRE(bIsEmpty);
 
 			Queue.Enqueue(42);
 			REQUIRE(!Queue.IsEmpty());
 
-			Result = AutoRTFM::Transact([&]
-				{
-					bIsEmpty = Queue.IsEmpty();
-					AutoRTFM::AbortTransaction();
-				});
+			AutoRTFM::Testing::Abort([&]
+			{
+				bIsEmpty = Queue.IsEmpty();
+				AutoRTFM::AbortTransaction();
+			});
 
-			REQUIRE(AutoRTFM::ETransactionResult::AbortedByRequest == Result);
 			REQUIRE(bIsEmpty);
 
-			Result = AutoRTFM::Transact([&]
-				{
-					bIsEmpty = Queue.IsEmpty();
-				});
+			AutoRTFM::Testing::Commit([&]
+			{
+				bIsEmpty = Queue.IsEmpty();
+			});
 
-			REQUIRE(AutoRTFM::ETransactionResult::Committed == Result);
 			REQUIRE(!bIsEmpty);
 		}
 
@@ -2876,21 +2651,19 @@ TEST_CASE("UECore.TQueue")
 			TQueue<int, EQueueMode::SingleThreaded> Queue;
 			REQUIRE(Queue.Enqueue(42));
 
-			AutoRTFM::ETransactionResult Result = AutoRTFM::Transact([&]
-				{
-					*Queue.Peek() = 13;
-					AutoRTFM::AbortTransaction();
-				});
+			AutoRTFM::Testing::Abort([&]
+			{
+				*Queue.Peek() = 13;
+				AutoRTFM::AbortTransaction();
+			});
 
-			REQUIRE(AutoRTFM::ETransactionResult::AbortedByRequest == Result);
 			REQUIRE(42 == *Queue.Peek());
 
-			Result = AutoRTFM::Transact([&]
-				{
-					*Queue.Peek() = 13;
-				});
+			AutoRTFM::Testing::Commit([&]
+			{
+				*Queue.Peek() = 13;
+			});
 
-			REQUIRE(AutoRTFM::ETransactionResult::Committed == Result);
 			REQUIRE(13 == *Queue.Peek());
 		}
 
@@ -2902,21 +2675,19 @@ TEST_CASE("UECore.TQueue")
 
 				bool bSucceeded = true;
 
-				AutoRTFM::ETransactionResult Result = AutoRTFM::Transact([&]
-					{
-						bSucceeded = Queue.Pop();
-						AutoRTFM::AbortTransaction();
-					});
+				AutoRTFM::Testing::Abort([&]
+				{
+					bSucceeded = Queue.Pop();
+					AutoRTFM::AbortTransaction();
+				});
 
-				REQUIRE(AutoRTFM::ETransactionResult::AbortedByRequest == Result);
 				REQUIRE(bSucceeded);
 
-				Result = AutoRTFM::Transact([&]
-					{
-						bSucceeded = Queue.Pop();
-					});
+				AutoRTFM::Testing::Commit([&]
+				{
+					bSucceeded = Queue.Pop();
+				});
 
-				REQUIRE(AutoRTFM::ETransactionResult::Committed == Result);
 				REQUIRE(!bSucceeded);
 			}
 
@@ -2927,22 +2698,20 @@ TEST_CASE("UECore.TQueue")
 
 				bool bSucceeded = false;
 
-				AutoRTFM::ETransactionResult Result = AutoRTFM::Transact([&]
-					{
-						bSucceeded = Queue.Pop();
-						AutoRTFM::AbortTransaction();
-					});
+				AutoRTFM::Testing::Abort([&]
+				{
+					bSucceeded = Queue.Pop();
+					AutoRTFM::AbortTransaction();
+				});
 
-				REQUIRE(AutoRTFM::ETransactionResult::AbortedByRequest == Result);
 				REQUIRE(!bSucceeded);
 				REQUIRE(!Queue.IsEmpty());
 
-				Result = AutoRTFM::Transact([&]
-					{
-						bSucceeded = Queue.Pop();
-					});
+				AutoRTFM::Testing::Commit([&]
+				{
+					bSucceeded = Queue.Pop();
+				});
 
-				REQUIRE(AutoRTFM::ETransactionResult::Committed == Result);
 				REQUIRE(bSucceeded);
 				REQUIRE(Queue.IsEmpty());
 			}
@@ -2960,10 +2729,10 @@ TEST_CASE("UECore.FConfigFile")
 
 		REQUIRE(!Config.IsEmpty());
 	
-		AutoRTFM::Commit([&]
-			{
-				Config.Empty();
-			});
+		AutoRTFM::Testing::Commit([&]
+		{
+			Config.Empty();
+		});
 		
 		REQUIRE(Config.IsEmpty());
 	}
