[AutoRTFM] Improve tests for transactionally-safe RWLocks.

The code has been moved out of UECoreTests since it's several hundred lines long.

Several of the existing test cases were identical, except for a single type/value. These were flattened using TEMPLATE_TEST_CASE. There is actually potential for further simplification, as well, since the tests for UE::TRWScopeLock are largely overlapping with UE::TReadScopeLock/UE::TWriteScopeLock.

New test cases were added to fill in various gaps called out in the initial review. This exposed some new bugs which will need to be addressed. (See https://jira.it.epicgames.com/browse/SOL-7658 ).

diff --git a/Engine/Source/Programs/AutoRTFMTests/Private/RWLockTests.cpp b/Engine/Source/Programs/AutoRTFMTests/Private/RWLockTests.cpp
new file mode 100644
--- /dev/null
+++ b/Engine/Source/Programs/AutoRTFMTests/Private/RWLockTests.cpp
@@ -0,0 +1,457 @@
+// Copyright Epic Games, Inc. All Rights Reserved.
+
+#include "Catch2Includes.h"
+#include "AutoRTFM/AutoRTFM.h"
+#include "AutoRTFM/AutoRTFMTesting.h"
+#include "Misc/TransactionallySafeRWLock.h"
+#include "Misc/ScopeRWLock.h"
+#include "Utils.h"
+#include <thread>
+#include <mutex>
+
+// Helper class for TEMPLATE_TEST_CASE, which expects a type but we want to pass a value.
+template <FRWScopeLockType T>
+struct TRWLockTypeHolder
+{
+	static constexpr FRWScopeLockType Value = T;
+};
+
+TEMPLATE_TEST_CASE("TransactionallySafeRWLock.RWScopeLock", "[ScopeLock]", TRWLockTypeHolder<FRWScopeLockType::SLT_ReadOnly>, TRWLockTypeHolder<FRWScopeLockType::SLT_Write>)
+{
+	SECTION("Outside Transaction With Lock")
+	{
+		FTransactionallySafeRWLock RWLock;
+
+		AutoRTFM::Testing::Abort([&]
+		{
+			UE::TRWScopeLock Lock(RWLock, TestType::Value);
+			AutoRTFM::AbortTransaction();
+		});
+
+		AutoRTFM::Testing::Commit([&]
+		{
+			UE::TRWScopeLock Lock(RWLock, TestType::Value);
+		});
+	}
+
+	SECTION("Inside Transaction With Lock")
+	{
+		AutoRTFM::Testing::Abort([&]
+		{
+			FTransactionallySafeRWLock RWLock;
+			UE::TRWScopeLock Lock(RWLock, TestType::Value);
+			AutoRTFM::AbortTransaction();
+		});
+
+		AutoRTFM::Testing::Commit([&]
+		{
+			FTransactionallySafeRWLock RWLock;
+			UE::TRWScopeLock Lock(RWLock, TestType::Value);
+		});
+	}
+}
+
+TEMPLATE_TEST_CASE("TransactionallySafeRWLock.NestedRWScopeLock", "[ScopeLock]", TRWLockTypeHolder<FRWScopeLockType::SLT_ReadOnly>, TRWLockTypeHolder<FRWScopeLockType::SLT_Write>)
+{
+	SECTION("Abort(Abort(Lock))")
+	{
+		AutoRTFM::Testing::Abort([&]
+		{
+			FTransactionallySafeRWLock RWLock;
+
+			AutoRTFM::Testing::Abort([&]
+			{
+				UE::TRWScopeLock Lock(RWLock, TestType::Value);
+				AutoRTFM::CascadingAbortTransaction();
+			});
+		});
+	}
+
+	SECTION("Commit(Abort(Lock))")
+	{
+		AutoRTFM::Testing::Commit([&]
+		{
+			FTransactionallySafeRWLock RWLock;
+
+			AutoRTFM::Testing::Abort([&]
+			{
+				UE::TRWScopeLock Lock(RWLock, TestType::Value);
+				AutoRTFM::AbortTransaction();
+			});
+		});
+	}
+
+	SECTION("Abort(Commit(Lock))")
+	{
+		AutoRTFM::Testing::Abort([&]
+		{
+			FTransactionallySafeRWLock RWLock;
+
+			AutoRTFM::Testing::Commit([&]
+			{
+				UE::TRWScopeLock Lock(RWLock, TestType::Value);
+			});
+
+			AutoRTFM::AbortTransaction();
+		});
+	}
+
+	SECTION("Commit(Commit(Lock))")
+	{
+		AutoRTFM::Testing::Commit([&]
+		{
+			FTransactionallySafeRWLock RWLock;
+
+			AutoRTFM::Testing::Commit([&]
+			{
+				UE::TRWScopeLock Lock(RWLock, TestType::Value);
+			});
+		});
+	}
+}
+
+TEMPLATE_TEST_CASE("TransactionallySafeRWLock.ScopeLock", "[ScopeLock]", UE::TReadScopeLock<FTransactionallySafeRWLock>, UE::TWriteScopeLock<FTransactionallySafeRWLock>)
+{
+	SECTION("Outside Transaction")
+	{
+		FTransactionallySafeRWLock RWLock;
+
+		AutoRTFM::Testing::Abort([&]
+		{
+			TestType Lock(RWLock);
+			AutoRTFM::AbortTransaction();
+		});
+
+		AutoRTFM::Testing::Commit([&]
+		{
+			TestType Lock(RWLock);
+		});
+	}
+
+	SECTION("Inside Transaction")
+	{
+		AutoRTFM::Testing::Abort([&]
+		{
+			FTransactionallySafeRWLock RWLock;
+			TestType Lock(RWLock);
+			AutoRTFM::AbortTransaction();
+		});
+
+		AutoRTFM::Testing::Commit([&]
+		{
+			FTransactionallySafeRWLock RWLock;
+			TestType Lock(RWLock);
+		});
+	}
+}
+
+TEMPLATE_TEST_CASE("TransactionallySafeRWLock.NestedScopeLock", "[ScopeLock]", UE::TReadScopeLock<FTransactionallySafeRWLock>, UE::TWriteScopeLock<FTransactionallySafeRWLock>)
+{
+	SECTION("Abort(Abort(Lock))")
+	{
+		AutoRTFM::Testing::Abort([&]
+		{
+			FTransactionallySafeRWLock RWLock;
+
+			AutoRTFM::Testing::Abort([&]
+			{
+				TestType Lock(RWLock);
+				AutoRTFM::CascadingAbortTransaction();
+			});
+		});
+	}
+
+	SECTION("Commit(Abort(Lock))")
+	{
+		AutoRTFM::Testing::Commit([&]
+		{
+			FTransactionallySafeRWLock RWLock;
+
+			AutoRTFM::Testing::Abort([&]
+			{
+				TestType Lock(RWLock);
+				AutoRTFM::AbortTransaction();
+			});
+		});
+	}
+
+	SECTION("Abort(Commit(Lock))")
+	{
+		AutoRTFM::Testing::Abort([&]
+		{
+			FTransactionallySafeRWLock RWLock;
+
+			AutoRTFM::Testing::Commit([&]
+			{
+				TestType Lock(RWLock);
+			});
+
+			AutoRTFM::AbortTransaction();
+		});
+	}
+
+	SECTION("Commit(Commit(Lock))")
+	{
+		AutoRTFM::Testing::Commit([&]
+		{
+			FTransactionallySafeRWLock RWLock;
+
+			AutoRTFM::Testing::Commit([&]
+			{
+				TestType Lock(RWLock);
+			});
+		});
+	}
+}
+
+TEMPLATE_TEST_CASE("TransactionallySafeRWLock.ScopeLock.FailedTryWriteInTransaction", "[ScopeLock]", UE::TReadScopeLock<FTransactionallySafeRWLock>, UE::TWriteScopeLock<FTransactionallySafeRWLock>)
+{
+	FTransactionallySafeRWLock RWLock;
+	TestType Lock(RWLock);
+
+	SECTION("Abort")
+	{
+		AutoRTFM::Testing::Abort([&]
+		{
+			REQUIRE(!RWLock.TryWriteLock());
+			AutoRTFM::CascadingAbortTransaction();
+		});
+	}
+
+	SECTION("Commit")
+	{
+		AutoRTFM::Testing::Commit([&]
+		{
+			REQUIRE(!RWLock.TryWriteLock());
+		});
+	}
+}
+
+TEMPLATE_TEST_CASE("TransactionallySafeRWLock.ZeroAndReconstruct", "[ScopeLock]", UE::TReadScopeLock<FTransactionallySafeRWLock>, UE::TWriteScopeLock<FTransactionallySafeRWLock>)
+{
+	SECTION("Lock, destruct, memzero, reconstruct")
+	{
+		FTransactionallySafeRWLock RWLock;
+		SECTION("Commit")
+		{
+			AutoRTFM::Testing::Commit([&]
+			{
+				{
+					// Lock and then unlock
+					TestType Lock(RWLock);
+				}
+				RWLock.~FTransactionallySafeRWLock();
+				memset(&RWLock, 0, sizeof(RWLock));
+				new (&RWLock) FTransactionallySafeRWLock();
+			});
+		}
+		SECTION("Abort")
+		{
+			AutoRTFM::Testing::Abort([&]
+			{
+				{
+					// Lock and then unlock
+					TestType Lock(RWLock);
+				}
+				RWLock.~FTransactionallySafeRWLock();
+				memset(&RWLock, 0, sizeof(RWLock));
+				new (&RWLock) FTransactionallySafeRWLock();
+				AutoRTFM::AbortTransaction();
+			});
+		}
+	}
+}
+
+TEST_CASE("TransactionallySafeRWLock.TryWriteLock")
+{
+	SECTION("Declared Outside Transaction")
+	{
+		FTransactionallySafeRWLock RWLock;
+
+		AutoRTFM::Testing::Abort([&]
+		{
+			REQUIRE(RWLock.TryWriteLock());
+			AutoRTFM::AbortTransaction();
+		});
+
+		AutoRTFM::Testing::Commit([&]
+		{
+			REQUIRE(RWLock.TryWriteLock());
+			RWLock.WriteUnlock();
+		});
+	}
+
+	SECTION("Declared Inside Transaction")
+	{
+		AutoRTFM::Testing::Abort([&]
+		{
+			FTransactionallySafeRWLock RWLock;
+			REQUIRE(RWLock.TryWriteLock());
+			AutoRTFM::AbortTransaction();
+		});
+
+		AutoRTFM::Testing::Commit([&]
+		{
+			FTransactionallySafeRWLock RWLock;
+			REQUIRE(RWLock.TryWriteLock());
+			RWLock.WriteUnlock();
+		});
+	}
+}
+
+TEST_CASE("TransactionallySafeRWLock.NestedTryWriteLock")
+{
+	SECTION("Abort(Abort(TryWriteLock))")
+	{
+		AutoRTFM::Testing::Abort([&]
+		{
+			FTransactionallySafeRWLock RWLock;
+
+			AutoRTFM::Testing::Abort([&]
+			{
+				REQUIRE(RWLock.TryWriteLock());
+				AutoRTFM::CascadingAbortTransaction();
+			});
+		});
+	}
+
+	SECTION("Commit(Abort(TryWriteLock))")
+	{
+		AutoRTFM::Testing::Commit([&]
+		{
+			FTransactionallySafeRWLock RWLock;
+
+			AutoRTFM::Testing::Abort([&]
+			{
+				REQUIRE(RWLock.TryWriteLock());
+				AutoRTFM::AbortTransaction();
+			});
+		});
+	}
+
+	SECTION("Abort(Commit(TryWriteLock))")
+	{
+		AutoRTFM::Testing::Abort([&]
+		{
+			FTransactionallySafeRWLock RWLock;
+
+			AutoRTFM::Testing::Commit([&]
+			{
+				REQUIRE(RWLock.TryWriteLock());
+			});
+
+			AutoRTFM::AbortTransaction();
+		});
+	}
+
+	SECTION("Commit(Commit(TryWriteLock))")
+	{
+		AutoRTFM::Testing::Commit([&]
+		{
+			FTransactionallySafeRWLock RWLock;
+
+			AutoRTFM::Testing::Commit([&]
+			{
+				REQUIRE(RWLock.TryWriteLock());
+			});
+
+			RWLock.WriteUnlock();
+		});
+	}
+}
+
+TEMPLATE_TEST_CASE("TransactionallySafeRWLock.RWScopeLock.ScopeLockThenTryWriteLock", "[ScopeLock]", UE::TReadScopeLock<FTransactionallySafeRWLock>, UE::TWriteScopeLock<FTransactionallySafeRWLock>)
+{
+	SECTION("Lock, Abort(TryWrite)")
+	{
+		FTransactionallySafeRWLock RWLock;
+		TestType Lock(RWLock);
+
+		AutoRTFM::Testing::Abort([&]
+		{
+			REQUIRE(!RWLock.TryWriteLock());
+			AutoRTFM::CascadingAbortTransaction();
+		});
+	}
+
+	SECTION("Lock, Commit(TryWrite)")
+	{
+		FTransactionallySafeRWLock RWLock;
+		TestType Lock(RWLock);
+
+		AutoRTFM::Testing::Commit([&]
+		{
+			REQUIRE(!RWLock.TryWriteLock());
+		});
+	}
+}
+
+TEST_CASE("TransactionallySafeRWLock.MultipleTryWriteLocks")
+{
+	SECTION("TryWrite, TryWrite, Unlock")
+	{
+		FTransactionallySafeRWLock RWLock;
+		REQUIRE(RWLock.TryWriteLock());
+		REQUIRE(!RWLock.TryWriteLock());
+		RWLock.WriteUnlock();
+	}
+
+	SECTION("TryWrite, Commit(TryWrite), Unlock")
+	{
+		FTransactionallySafeRWLock RWLock;
+		REQUIRE(RWLock.TryWriteLock());
+		AutoRTFM::Testing::Commit([&]
+		{
+			REQUIRE(!RWLock.TryWriteLock());
+		});
+		RWLock.WriteUnlock();
+	}
+
+	// Test case does not pass; RWLock trashes its internal state: #jira SOL-7658
+	SECTION("TryWrite, Commit(TryWrite, Unlock)") 
+	{
+//		FTransactionallySafeRWLock RWLock;
+//		REQUIRE(RWLock.TryWriteLock());
+//		AutoRTFM::Testing::Commit([&]
+//		{
+//			REQUIRE(!RWLock.TryWriteLock());
+//			RWLock.WriteUnlock();
+//		}); 
+	}
+
+	SECTION("TryWrite, Abort(TryWrite), Unlock")
+	{
+		FTransactionallySafeRWLock RWLock;
+		REQUIRE(RWLock.TryWriteLock());
+		AutoRTFM::Testing::Abort([&]
+		{
+			REQUIRE(!RWLock.TryWriteLock());
+			AutoRTFM::AbortTransaction();
+		});
+		RWLock.WriteUnlock();
+	}
+
+	// Test case does not pass; RWLock trashes its internal state: #jira SOL-7658
+	SECTION("Commit(TryWrite, TryWrite), Unlock")
+	{
+//		FTransactionallySafeRWLock RWLock;
+//		AutoRTFM::Testing::Commit([&]
+//		{
+//			REQUIRE(RWLock.TryWriteLock());
+//			REQUIRE(!RWLock.TryWriteLock());
+//		});
+//		RWLock.WriteUnlock();
+	}
+
+	// Test case does not pass; RWLock trashes its internal state: #jira SOL-7658
+	SECTION("Commit(TryWrite, TryWrite, Unlock)")
+	{
+//		AutoRTFM::Testing::Commit([&]
+//		{
+//			FTransactionallySafeRWLock RWLock;
+//			REQUIRE(RWLock.TryWriteLock());
+//			REQUIRE(!RWLock.TryWriteLock());
+//			RWLock.WriteUnlock();
+//		});
+	}
+}
diff --git a/Engine/Source/Programs/AutoRTFMTests/Private/UECoreTests.cpp b/Engine/Source/Programs/AutoRTFMTests/Private/UECoreTests.cpp
--- a/Engine/Source/Programs/AutoRTFMTests/Private/UECoreTests.cpp
+++ b/Engine/Source/Programs/AutoRTFMTests/Private/UECoreTests.cpp
@@ -1115,493 +1115,6 @@ TEST_CASE("UECore.TScopeLock_TransactionallySafeCriticalSection")
 	}
 }
 
-TEST_CASE("UECore.TransactionallySafeRWScopeLock")
-{
-	SECTION("Outside Transaction With Read Lock")
-	{
-		FTransactionallySafeRWLock RWLock;
-
-		AutoRTFM::Testing::Abort([&]
-		{
-			UE::TRWScopeLock Lock(RWLock, FRWScopeLockType::SLT_ReadOnly);
-			AutoRTFM::AbortTransaction();
-		});
-
-		AutoRTFM::Testing::Commit([&]
-		{
-			UE::TRWScopeLock Lock(RWLock, FRWScopeLockType::SLT_ReadOnly);
-		});
-	}
-
-	SECTION("Inside Transaction With Read Lock")
-	{
-		AutoRTFM::Testing::Abort([&]
-		{
-			FTransactionallySafeRWLock RWLock;
-			UE::TRWScopeLock Lock(RWLock, FRWScopeLockType::SLT_ReadOnly);
-			AutoRTFM::AbortTransaction();
-		});
-
-		AutoRTFM::Testing::Commit([&]
-		{
-			FTransactionallySafeRWLock RWLock;
-			UE::TRWScopeLock Lock(RWLock, FRWScopeLockType::SLT_ReadOnly);
-		});
-	}
-
-	SECTION("Inside Transaction Used In Nested Transaction With Read Lock")
-	{
-		AutoRTFM::Testing::Abort([&]
-		{
-			FTransactionallySafeRWLock RWLock;
-
-			AutoRTFM::Testing::Abort([&]
-			{
-				UE::TRWScopeLock Lock(RWLock, FRWScopeLockType::SLT_ReadOnly);
-				AutoRTFM::CascadingAbortTransaction();
-			});
-		});
-
-		AutoRTFM::Testing::Commit([&]
-		{
-			FTransactionallySafeRWLock RWLock;
-
-			AutoRTFM::Testing::Abort([&]
-			{
-				UE::TRWScopeLock Lock(RWLock, FRWScopeLockType::SLT_ReadOnly);
-				AutoRTFM::AbortTransaction();
-			});
-		});
-
-		AutoRTFM::Testing::Abort([&]
-		{
-			FTransactionallySafeRWLock RWLock;
-
-			AutoRTFM::Testing::Commit([&]
-			{
-				UE::TRWScopeLock Lock(RWLock, FRWScopeLockType::SLT_ReadOnly);
-			});
-
-			AutoRTFM::AbortTransaction();
-		});
-
-		AutoRTFM::Testing::Commit([&]
-		{
-			FTransactionallySafeRWLock RWLock;
-
-			AutoRTFM::Testing::Commit([&]
-			{
-				UE::TRWScopeLock Lock(RWLock, FRWScopeLockType::SLT_ReadOnly);
-			});
-		});
-	}
-
-	SECTION("Outside Transaction With Write Lock")
-	{
-		FTransactionallySafeRWLock RWLock;
-
-		AutoRTFM::Testing::Abort([&]
-		{
-			UE::TRWScopeLock Lock(RWLock, FRWScopeLockType::SLT_Write);
-			AutoRTFM::AbortTransaction();
-		});
-
-		AutoRTFM::Testing::Commit([&]
-		{
-			UE::TRWScopeLock Lock(RWLock, FRWScopeLockType::SLT_Write);
-		});
-	}
-
-	SECTION("Inside Transaction With Write Lock")
-	{
-		AutoRTFM::Testing::Abort([&]
-		{
-			FTransactionallySafeRWLock RWLock;
-			UE::TRWScopeLock Lock(RWLock, FRWScopeLockType::SLT_Write);
-			AutoRTFM::AbortTransaction();
-		});
-
-		AutoRTFM::Testing::Commit([&]
-		{
-			FTransactionallySafeRWLock RWLock;
-			UE::TRWScopeLock Lock(RWLock, FRWScopeLockType::SLT_Write);
-		});
-	}
-
-	SECTION("Inside Transaction Used In Nested Transaction With Write Lock")
-	{
-		AutoRTFM::Testing::Abort([&]
-		{
-			FTransactionallySafeRWLock RWLock;
-
-			AutoRTFM::Testing::Abort([&]
-			{
-				UE::TRWScopeLock Lock(RWLock, FRWScopeLockType::SLT_Write);
-				AutoRTFM::CascadingAbortTransaction();
-			});
-		});
-
-		AutoRTFM::Testing::Commit([&]
-		{
-			FTransactionallySafeRWLock RWLock;
-
-			AutoRTFM::Testing::Abort([&]
-			{
-				UE::TRWScopeLock Lock(RWLock, FRWScopeLockType::SLT_Write);
-				AutoRTFM::AbortTransaction();
-			});
-		});
-
-		AutoRTFM::Testing::Abort([&]
-		{
-			FTransactionallySafeRWLock RWLock;
-
-			AutoRTFM::Testing::Commit([&]
-			{
-				UE::TRWScopeLock Lock(RWLock, FRWScopeLockType::SLT_Write);
-			});
-
-			AutoRTFM::AbortTransaction();
-		});
-
-		AutoRTFM::Testing::Commit([&]
-		{
-			FTransactionallySafeRWLock RWLock;
-
-			AutoRTFM::Testing::Commit([&]
-			{
-				UE::TRWScopeLock Lock(RWLock, FRWScopeLockType::SLT_Write);
-			});
-		});
-	}
-}
-
-TEST_CASE("UECore.TransactionallySafeReadScopeLock")
-{
-	SECTION("Outside Transaction")
-	{
-		FTransactionallySafeRWLock RWLock;
-
-		AutoRTFM::Testing::Abort([&]
-		{
-			UE::TReadScopeLock Lock(RWLock);
-			AutoRTFM::AbortTransaction();
-		});
-
-		AutoRTFM::Testing::Commit([&]
-		{
-			UE::TReadScopeLock Lock(RWLock);
-		});
-	}
-
-	SECTION("Inside Transaction")
-	{
-		AutoRTFM::Testing::Abort([&]
-		{
-			FTransactionallySafeRWLock RWLock;
-			UE::TReadScopeLock Lock(RWLock);
-			AutoRTFM::AbortTransaction();
-		});
-
-		AutoRTFM::Testing::Commit([&]
-		{
-			FTransactionallySafeRWLock RWLock;
-			UE::TReadScopeLock Lock(RWLock);
-		});
-	}
-
-	SECTION("Inside Transaction Used In Nested Transaction")
-	{
-		AutoRTFM::Testing::Abort([&]
-		{
-			FTransactionallySafeRWLock RWLock;
-
-			AutoRTFM::Testing::Abort([&]
-			{
-				UE::TReadScopeLock Lock(RWLock);
-				AutoRTFM::CascadingAbortTransaction();
-			});
-		});
-
-		AutoRTFM::Testing::Commit([&]
-		{
-			FTransactionallySafeRWLock RWLock;
-
-			AutoRTFM::Testing::Abort([&]
-			{
-				UE::TReadScopeLock Lock(RWLock);
-				AutoRTFM::AbortTransaction();
-			});
-		});
-
-		AutoRTFM::Testing::Abort([&]
-		{
-			FTransactionallySafeRWLock RWLock;
-
-			AutoRTFM::Testing::Commit([&]
-			{
-				UE::TReadScopeLock Lock(RWLock);
-			});
-
-			AutoRTFM::AbortTransaction();
-		});
-
-		AutoRTFM::Testing::Commit([&]
-		{
-			FTransactionallySafeRWLock RWLock;
-
-			AutoRTFM::Testing::Commit([&]
-			{
-				UE::TReadScopeLock Lock(RWLock);
-			});
-		});
-	}
-
-	SECTION("TReadScopeLock, destruct, memzero, reconstruct")
-	{
-		FTransactionallySafeRWLock RWLock;
-		SECTION("Commit")
-		{
-			AutoRTFM::Testing::Commit([&]
-			{
-				{
-					// Read-lock and then unlock
-					UE::TReadScopeLock Lock(RWLock);
-				}
-				RWLock.~FTransactionallySafeRWLock();
-				memset(&RWLock, 0, sizeof(RWLock));
-				new (&RWLock) FTransactionallySafeRWLock();
-			});
-		}
-		SECTION("Abort")
-		{
-			AutoRTFM::Testing::Abort([&]
-			{
-				{
-					// Read-Lock and then unlock
-					UE::TReadScopeLock Lock(RWLock);
-				}
-				RWLock.~FTransactionallySafeRWLock();
-				memset(&RWLock, 0, sizeof(RWLock));
-				new (&RWLock) FTransactionallySafeRWLock();
-				AutoRTFM::AbortTransaction();
-			});
-		}
-	}
-}
-
-TEST_CASE("UECore.TransactionallySafeRWLockTryWriteLock")
-{
-	SECTION("Outside Transaction")
-	{
-		FTransactionallySafeRWLock RWLock;
-
-		AutoRTFM::Testing::Abort([&]
-		{
-			REQUIRE(RWLock.TryWriteLock());
-			AutoRTFM::AbortTransaction();
-		});
-
-		AutoRTFM::Testing::Commit([&]
-		{
-			REQUIRE(RWLock.TryWriteLock());
-			RWLock.WriteUnlock();
-		});
-	}
-
-	SECTION("Inside Transaction")
-	{
-		AutoRTFM::Testing::Abort([&]
-		{
-			FTransactionallySafeRWLock RWLock;
-			REQUIRE(RWLock.TryWriteLock());
-			AutoRTFM::AbortTransaction();
-		});
-
-		AutoRTFM::Testing::Commit([&]
-		{
-			FTransactionallySafeRWLock RWLock;
-			REQUIRE(RWLock.TryWriteLock());
-			RWLock.WriteUnlock();
-		});
-	}
-
-	SECTION("Inside Transaction Used In Nested Transaction")
-	{
-		AutoRTFM::Testing::Abort([&]
-		{
-			FTransactionallySafeRWLock RWLock;
-
-			AutoRTFM::Testing::Abort([&]
-			{
-				REQUIRE(RWLock.TryWriteLock());
-				AutoRTFM::CascadingAbortTransaction();
-			});
-		});
-
-		AutoRTFM::Testing::Commit([&]
-		{
-			FTransactionallySafeRWLock RWLock;
-
-			AutoRTFM::Testing::Abort([&]
-			{
-				REQUIRE(RWLock.TryWriteLock());
-				AutoRTFM::AbortTransaction();
-			});
-		});
-
-		AutoRTFM::Testing::Abort([&]
-		{
-			FTransactionallySafeRWLock RWLock;
-
-			AutoRTFM::Testing::Commit([&]
-			{
-				REQUIRE(RWLock.TryWriteLock());
-			});
-
-			AutoRTFM::AbortTransaction();
-		});
-
-		AutoRTFM::Testing::Commit([&]
-		{
-			FTransactionallySafeRWLock RWLock;
-
-			AutoRTFM::Testing::Commit([&]
-			{
-				REQUIRE(RWLock.TryWriteLock());
-			});
-
-			RWLock.WriteUnlock();
-		});
-	}
-
-	SECTION("Multiple TryWriteLocks In Transaction")
-	{
-		AutoRTFM::Testing::Commit([&]
-		{
-			FTransactionallySafeRWLock RWLock;
-			REQUIRE(RWLock.TryWriteLock());
-			REQUIRE(RWLock.TryWriteLock());
-			RWLock.WriteUnlock();
-		});
-	}
-}
-
-
-TEST_CASE("UECore.TransactionallySafeWriteScopeLock")
-{
-	SECTION("Outside Transaction")
-	{
-		FTransactionallySafeRWLock RWLock;
-
-		AutoRTFM::Testing::Abort([&]
-		{
-			UE::TWriteScopeLock Lock(RWLock);
-			AutoRTFM::AbortTransaction();
-		});
-
-		AutoRTFM::Testing::Commit([&]
-		{
-			UE::TWriteScopeLock Lock(RWLock);
-		});
-	}
-
-	SECTION("Inside Transaction")
-	{
-		AutoRTFM::Testing::Abort([&]
-		{
-			FTransactionallySafeRWLock RWLock;
-			UE::TWriteScopeLock Lock(RWLock);
-			AutoRTFM::AbortTransaction();
-		});
-
-		AutoRTFM::Testing::Commit([&]
-		{
-			FTransactionallySafeRWLock RWLock;
-			UE::TWriteScopeLock Lock(RWLock);
-		});
-	}
-
-	SECTION("Inside Transaction Used In Nested Transaction")
-	{
-		AutoRTFM::Testing::Abort([&]
-		{
-			FTransactionallySafeRWLock RWLock;
-
-			AutoRTFM::Testing::Abort([&]
-			{
-				UE::TWriteScopeLock Lock(RWLock);
-				AutoRTFM::CascadingAbortTransaction();
-			});
-		});
-
-		AutoRTFM::Testing::Commit([&]
-		{
-			FTransactionallySafeRWLock RWLock;
-
-			AutoRTFM::Testing::Abort([&]
-			{
-				UE::TWriteScopeLock Lock(RWLock);
-				AutoRTFM::AbortTransaction();
-			});
-		});
-
-		AutoRTFM::Testing::Abort([&]
-		{
-			FTransactionallySafeRWLock RWLock;
-
-			AutoRTFM::Testing::Commit([&]
-			{
-				UE::TWriteScopeLock Lock(RWLock);
-			});
-
-			AutoRTFM::AbortTransaction();
-		});
-
-		AutoRTFM::Testing::Commit([&]
-		{
-			FTransactionallySafeRWLock RWLock;
-
-			AutoRTFM::Testing::Commit([&]
-			{
-				UE::TWriteScopeLock Lock(RWLock);
-			});
-		});
-	}
-
-	SECTION("TWriteScopeLock, destruct, memzero, reconstruct")
-	{
-		FTransactionallySafeRWLock RWLock;
-		SECTION("Commit")
-		{
-			AutoRTFM::Testing::Commit([&]
-			{
-				{
-					// Write-lock and then unlock
-					UE::TWriteScopeLock Lock(RWLock);
-				}
-				RWLock.~FTransactionallySafeRWLock();
-				memset(&RWLock, 0, sizeof(RWLock));
-				new (&RWLock) FTransactionallySafeRWLock();
-			});
-		}
-		SECTION("Abort")
-		{
-			AutoRTFM::Testing::Abort([&]
-			{
-				{
-					// Write-lock and then unlock
-					UE::TWriteScopeLock Lock(RWLock);
-				}
-				RWLock.~FTransactionallySafeRWLock();
-				memset(&RWLock, 0, sizeof(RWLock));
-				new (&RWLock) FTransactionallySafeRWLock();
-				AutoRTFM::AbortTransaction();
-			});
-		}
-	}
-}
-
 TEST_CASE("UECore.TransactionallySafeMutex")
 {
 	SECTION("Outside Transaction")
