[AutoRTFM] Add THashMap, a wrapper around std::unordered_map. Use it instead of TMap.

Part of the work to separate AutoRTFM out of core to break cyclic build dependency issues.

diff --git a/Engine/Source/Programs/AutoRTFMTests/Private/HashMapTests.cpp b/Engine/Source/Programs/AutoRTFMTests/Private/HashMapTests.cpp
new file mode 100644
--- /dev/null
+++ b/Engine/Source/Programs/AutoRTFMTests/Private/HashMapTests.cpp
@@ -0,0 +1,436 @@
+// Copyright Epic Games, Inc. All Rights Reserved.
+
+#include "AutoRTFM/HashMap.h"
+#include "Catch2Includes.h"
+#include "ObjectLifetimeHelper.h"
+#include <unordered_map>
+
+using HashMapIntToInt = AutoRTFM::THashMap<int, int>;
+using HashMapIntToObj = AutoRTFM::THashMap<int, AutoRTFMTestUtils::FObjectLifetimeHelper>;
+using HashMapObjToInt = AutoRTFM::THashMap<AutoRTFMTestUtils::FObjectLifetimeHelper, int>;
+using HashMapObjToObj = AutoRTFM::THashMap<AutoRTFMTestUtils::FObjectLifetimeHelper, AutoRTFMTestUtils::FObjectLifetimeHelper>;
+
+TEMPLATE_TEST_CASE("HashMap", "", HashMapIntToInt, HashMapIntToObj, HashMapObjToInt, HashMapObjToObj)
+{
+	REQUIRE(AutoRTFMTestUtils::FObjectLifetimeHelper::ConstructorCalls == 0);
+	REQUIRE(AutoRTFMTestUtils::FObjectLifetimeHelper::DestructorCalls == 0);
+
+	using MapType = TestType;
+	using KeyType = typename TestType::Key;
+	using ValueType = typename TestType::Value;
+
+	auto Check = [&](MapType& Map, const std::unordered_map<KeyType, ValueType>& Expected)
+	{
+		// Check the reported count is as expected.
+		REQUIRE(Map.Num() == Expected.size());
+		REQUIRE(Map.IsEmpty() == Expected.empty());
+
+		// Check the non-const begin() / end().
+		{
+			size_t I = 0;
+			for (AutoRTFM::TKeyAndValue<const KeyType, ValueType>& Item : Map)
+			{
+				I++;
+				auto ExpectedIt = Expected.find(Item.Key);
+				REQUIRE(ExpectedIt != Expected.end());
+				REQUIRE(ExpectedIt->second == Item.Value);
+				// Check mutability of the value
+				Item.Value = 99;
+				REQUIRE(Item.Value == 99);
+				Item.Value = ExpectedIt->second;
+			}
+			REQUIRE(I == Expected.size());
+		}
+		// Check the const begin() / end().
+		{
+			size_t I = 0;
+			for (const AutoRTFM::TKeyAndValue<const KeyType, ValueType>& Item : const_cast<const MapType&>(Map))
+			{
+				I++;
+				auto ExpectedIt = Expected.find(Item.Key);
+				REQUIRE(ExpectedIt != Expected.end());
+				REQUIRE(ExpectedIt->second == Item.Value);
+			}
+			REQUIRE(I == Expected.size());
+		}
+		// Check Contains() & Find() for all elements in Expected.
+		{
+			for (auto It : Expected)
+			{
+				REQUIRE(Map.Contains(It.first));
+				ValueType* FindValue = Map.Find(It.first);
+				REQUIRE(FindValue);
+				REQUIRE(*FindValue == It.second);
+			}
+		}
+	};
+
+	SECTION("Add")
+	{
+		MapType Map;
+		Check(Map, {});
+
+		Map.Add(10, 100);
+		Check(Map, {{10, 100}});
+
+		Map.Add(20, 200);
+		Check(Map, {{10, 100}, {20, 200}});
+
+		Map.Add(30, 300);
+		Check(Map, {{10, 100}, {20, 200}, {30, 300}});
+
+		Map.Add(20, 123);
+		Check(Map, {{10, 100}, {20, 123}, {30, 300}});
+
+		Map.Add(10, 321);
+		Check(Map, {{10, 321}, {20, 123}, {30, 300}});
+
+		Map.Add(40, 400);
+		Check(Map, {{10, 321}, {20, 123}, {30, 300}, {40, 400}});
+	}
+
+	SECTION("Find")
+	{
+		MapType Map;
+		Check(Map, {});
+
+		Map.Add(10, 100);
+		Map.Add(20, 200);
+		Map.Add(30, 300);
+
+		Check(Map, {{10, 100}, {20, 200}, {30, 300}});
+
+		{
+			ValueType* Value = Map.Find(KeyType{0});
+			REQUIRE(!Value);
+		}
+		{
+			ValueType* Value = Map.Find(KeyType{10});
+			REQUIRE((Value && *Value == 100));
+		}
+		{
+			ValueType* Value = Map.Find(KeyType{20});
+			REQUIRE((Value && *Value == 200));
+		}
+		{
+			ValueType* Value = Map.Find(KeyType{30});
+			REQUIRE((Value && *Value == 300));
+		}
+		{
+			ValueType* Value = Map.Find(KeyType{40});
+			REQUIRE(!Value);
+		}
+
+		Check(Map, {{10, 100}, {20, 200}, {30, 300}});
+	}
+
+	SECTION("FindOrAdd")
+	{
+		MapType Map;
+		Check(Map, {});
+
+		Map.Add(10, 100);
+		Map.Add(20, 200);
+		Map.Add(30, 300);
+		Check(Map, {{10, 100}, {20, 200}, {30, 300}});
+
+		REQUIRE(Map.FindOrAdd(0) == 0);
+		Check(Map, {{0, 0}, {10, 100}, {20, 200}, {30, 300}});
+		REQUIRE(Map.FindOrAdd(10) == 100);
+		Check(Map, {{0, 0}, {10, 100}, {20, 200}, {30, 300}});
+		REQUIRE(Map.FindOrAdd(20) == 200);
+		Check(Map, {{0, 0}, {10, 100}, {20, 200}, {30, 300}});
+		REQUIRE(Map.FindOrAdd(30) == 300);
+		Check(Map, {{0, 0}, {10, 100}, {20, 200}, {30, 300}});
+		REQUIRE(Map.FindOrAdd(40) == 0);
+		Check(Map, {{0, 0}, {10, 100}, {20, 200}, {30, 300}, {40, 0}});
+	}
+
+	SECTION("Remove")
+	{
+		MapType Map;
+		Check(Map, {});
+
+		Map.Add(10, 100);
+		Map.Add(20, 200);
+		Map.Add(30, 300);
+		Check(Map, {{10, 100}, {20, 200}, {30, 300}});
+
+		Map.Remove(0);
+		Check(Map, {{10, 100}, {20, 200}, {30, 300}});
+		Map.Remove(10);
+		Check(Map, {{20, 200}, {30, 300}});
+		Map.Remove(20);
+		Check(Map, {{30, 300}});
+		Map.Remove(30);
+		Check(Map, {});
+		Map.Remove(40);
+		Check(Map, {});
+	}
+
+	SECTION("Contains")
+	{
+		MapType Map;
+		Check(Map, {});
+
+		Map.Add(10, 100);
+		Map.Add(20, 200);
+		Map.Add(30, 300);
+		Check(Map, {{10, 100}, {20, 200}, {30, 300}});
+
+		REQUIRE(Map.Contains(0) == false);
+		REQUIRE(Map.Contains(10) == true);
+		REQUIRE(Map.Contains(20) == true);
+		REQUIRE(Map.Contains(30) == true);
+		REQUIRE(Map.Contains(40) == false);
+		Check(Map, {{10, 100}, {20, 200}, {30, 300}});
+	}
+
+	SECTION("Empty")
+	{
+		MapType Map;
+		Check(Map, {});
+
+		Map.Empty();
+		Check(Map, {});
+
+		Map.Add(10, 100);
+		Map.Add(20, 200);
+		Map.Add(30, 300);
+		Check(Map, {{10, 100}, {20, 200}, {30, 300}});
+
+		Map.Empty();
+		Check(Map, {});
+	}
+
+	SECTION("Reset")
+	{
+		MapType Map;
+		Check(Map, {});
+
+		Map.Reset();
+		Check(Map, {});
+
+		Map.Add(10, 100);
+		Map.Add(20, 200);
+		Map.Add(30, 300);
+		Check(Map, {{10, 100}, {20, 200}, {30, 300}});
+
+		Map.Reset();
+		Check(Map, {});
+	}
+
+	SECTION("Copy Construct")
+	{
+		SECTION("Empty")
+		{
+			MapType Source;
+			MapType Target(Source);
+			Check(Source, {});
+			Check(Target, {});
+		}
+
+		SECTION("Non-empty")
+		{
+			MapType Source;
+			Source.Add(10, 100);
+			Source.Add(20, 200);
+			Source.Add(30, 300);
+			MapType Target(Source);
+			Check(Source, {{10, 100}, {20, 200}, {30, 300}});
+			Check(Target, {{10, 100}, {20, 200}, {30, 300}});
+		}
+	}
+
+	SECTION("Move Construct")
+	{
+		SECTION("Empty")
+		{
+			MapType Source;
+			MapType Target(std::move(Source));
+			Check(Source, {});
+			Check(Target, {});
+		}
+
+		SECTION("Non-empty")
+		{
+			MapType Source;
+			Source.Add(10, 100);
+			Source.Add(20, 200);
+			Source.Add(30, 300);
+			MapType Target(std::move(Source));
+			Check(Source, {});
+			Check(Target, {{10, 100}, {20, 200}, {30, 300}});
+		}
+	}
+
+	SECTION("Copy Assign")
+	{
+		SECTION("Empty")
+		{
+			MapType Source;
+			MapType Target;
+			Target = Source;
+			Check(Source, {});
+			Check(Target, {});
+		}
+
+		SECTION("Non-empty")
+		{
+			MapType Source;
+			Source.Add(10, 100);
+			Source.Add(20, 200);
+			Source.Add(30, 300);
+			MapType Target;
+			Target = Source;
+			Check(Source, {{10, 100}, {20, 200}, {30, 300}});
+			Check(Target, {{10, 100}, {20, 200}, {30, 300}});
+		}
+	}
+
+	SECTION("Copy Assign Self")
+	{
+		SECTION("Empty")
+		{
+			MapType Map;
+			MapType& Alias = Map;
+			Map = Alias;
+			Check(Map, {});
+		}
+
+		SECTION("Non-empty")
+		{
+			MapType Map;
+			MapType& Alias = Map;
+			Map.Add(10, 100);
+			Map.Add(20, 200);
+			Map.Add(30, 300);
+			Map = Alias;
+			Check(Map, {{10, 100}, {20, 200}, {30, 300}});
+		}
+	}
+
+	SECTION("Move Assign")
+	{
+		SECTION("Empty")
+		{
+			MapType Source;
+			MapType Target;
+			Target = std::move(Source);
+			Check(Source, {});
+			Check(Target, {});
+		}
+
+		SECTION("Non-empty")
+		{
+			MapType Source;
+			Source.Add(10, 100);
+			Source.Add(20, 200);
+			Source.Add(30, 300);
+			MapType Target;
+			Target = std::move(Source);
+			Check(Source, {});
+			Check(Target, {{10, 100}, {20, 200}, {30, 300}});
+		}
+	}
+
+	SECTION("Move Assign Self")
+	{
+		SECTION("Empty")
+		{
+			MapType Map;
+			MapType& Alias = Map;
+			Map = std::move(Alias);
+			Check(Map, {});
+		}
+
+		SECTION("Non-empty")
+		{
+			MapType Map;
+			MapType& Alias = Map;
+			Map.Add(10, 100);
+			Map.Add(20, 200);
+			Map.Add(30, 300);
+			Map = std::move(Alias);
+			Check(Map, {{10, 100}, {20, 200}, {30, 300}});
+		}
+	}
+
+	SECTION("Soak")
+	{
+		MapType Map;
+		std::unordered_map<KeyType, ValueType> Expected;
+		for (int I = 0; I < 10000; I++)
+		{
+			int Mod100 = (I * 15485863) % 100;
+			KeyType Key = KeyType{(I * 804889) % 1000};
+			switch (Mod100)
+			{
+				case 0:
+				{
+					ValueType* Value = Map.Find(Key);
+					if (Value != nullptr)
+					{
+						auto ExpectedIt = Expected.find(Key);
+						REQUIRE(ExpectedIt != Expected.end());
+						REQUIRE(*Value == ExpectedIt->second);
+					}
+					else
+					{
+						REQUIRE(Expected.find(Key) == Expected.end());
+					}
+					break;
+				}
+				case 1:
+				{
+					ValueType Value = Map.FindOrAdd(Key);
+					auto ExpectedIt = Expected.find(Key);
+					if (ExpectedIt != Expected.end())
+					{
+						REQUIRE(Value == ExpectedIt->second);
+					}
+					else
+					{
+						REQUIRE(Value == ValueType{});
+						Expected[Key] = ValueType{};
+					}
+					break;
+				}
+				case 2:
+				{
+					Map.Remove(Key);
+					Expected.erase(Key);
+					break;
+				}
+				case 3:
+				{
+					REQUIRE(Map.Contains(Key) == (Expected.count(Key) != 0));
+					break;
+				}
+				case 4:
+				{
+					Map.Empty();
+					Expected.clear();
+					break;
+				}
+				case 5:
+				{
+					Map.Reset();
+					Expected.clear();
+					break;
+				}
+				default:
+				{
+					Map.Add(Key, ValueType{I});
+					Expected[Key] = ValueType{I};
+					break;
+				}
+			}
+			Check(Map, Expected);
+		}
+	}
+
+	REQUIRE(AutoRTFMTestUtils::FObjectLifetimeHelper::ConstructorCalls == AutoRTFMTestUtils::FObjectLifetimeHelper::DestructorCalls);
+	AutoRTFMTestUtils::FObjectLifetimeHelper::ConstructorCalls = 0;
+	AutoRTFMTestUtils::FObjectLifetimeHelper::DestructorCalls = 0;
+}
diff --git a/Engine/Source/Programs/AutoRTFMTests/Private/ObjectLifetimeHelper.h b/Engine/Source/Programs/AutoRTFMTests/Private/ObjectLifetimeHelper.h
--- a/Engine/Source/Programs/AutoRTFMTests/Private/ObjectLifetimeHelper.h
+++ b/Engine/Source/Programs/AutoRTFMTests/Private/ObjectLifetimeHelper.h
@@ -3,6 +3,7 @@
 #pragma once
 
 #include <cstddef>
+#include <functional>
 
 namespace AutoRTFMTestUtils
 {
@@ -27,3 +28,17 @@ struct FObjectLifetimeHelper
 };
 
 }  // namespace AutoRTFMTestUtils
+
+// std::hash specialization for AutoRTFMTestUtils::FObjectLifetimeHelper
+namespace std
+{
+	template<>
+	class hash<AutoRTFMTestUtils::FObjectLifetimeHelper>
+	{
+	public:
+		size_t operator()(const AutoRTFMTestUtils::FObjectLifetimeHelper& Object) const
+		{
+			return std::hash<int>{}(Object.Value);
+		}
+	};
+}
diff --git a/Engine/Source/Runtime/Core/Private/AutoRTFM/HashMap.h b/Engine/Source/Runtime/Core/Private/AutoRTFM/HashMap.h
new file mode 100644
--- /dev/null
+++ b/Engine/Source/Runtime/Core/Private/AutoRTFM/HashMap.h
@@ -0,0 +1,197 @@
+// Copyright Epic Games, Inc. All Rights Reserved.
+
+#pragma once
+
+#include "StlAllocator.h"
+#include <functional>
+#include <type_traits>
+#include <unordered_map>
+
+namespace AutoRTFM
+{
+
+// A simple structure holding a templated key and value pair.
+template<typename KeyType, typename ValueType>
+struct TKeyAndValue
+{
+	KeyType Key;
+	ValueType Value;
+};
+
+// HashMap is an unordered hashmap.
+// TODO(SOL-7652): This currently wraps a std::unordered_map. Reimplement to
+// improve performance and to avoid platform-specific variations in behavior.
+template<typename KeyType, typename ValueType, typename Hash = std::hash<KeyType>, typename Equal = std::equal_to<KeyType>>
+class THashMap
+{
+	using FInnerMap = std::unordered_map<KeyType, ValueType, Hash, Equal, StlAllocator<std::pair<const KeyType, ValueType>>>;
+
+	template<bool bConst>
+	class TIterator
+	{
+		using FInnerIterator = std::conditional_t<bConst, typename FInnerMap::const_iterator, typename FInnerMap::iterator>;
+		using FKeyAndValue = std::conditional_t<bConst, const TKeyAndValue<const KeyType, ValueType>, TKeyAndValue<const KeyType, ValueType>>;
+	public:
+		// Constructor wrapping inner-iterator
+		TIterator(FInnerIterator InnerIterator) : InnerIterator{InnerIterator} {}
+		// Copy constructor
+		TIterator(const TIterator&) = default;
+		// Move constructor
+		TIterator(TIterator&&) = default;
+		// Copy assignment operator
+		TIterator& operator = (const TIterator& Other) = default;
+		// Move assignment operator
+		TIterator& operator = (TIterator&& Other) = default;
+
+		FKeyAndValue& operator*() const
+		{
+			// Super hacky: Reinterpret a std::pair to a TKeyAndValue so
+			// that THashMap's iterator has the interface of a TMap instead of a 
+			// std::unordered_map.
+			// TODO(SOL-7652): Remove this filth when THashMap is reimplemented.
+			using FPair = std::conditional_t<bConst,
+				const std::pair<const KeyType, ValueType>,
+				std::pair<const KeyType, ValueType>>;
+			FPair* Pair = &*InnerIterator;
+			FKeyAndValue* KeyAndValue = reinterpret_cast<FKeyAndValue*>(Pair);
+			static_assert(std::is_same_v<decltype(Pair->first), decltype(KeyAndValue->Key)>);
+			static_assert(std::is_same_v<decltype(Pair->second), decltype(KeyAndValue->Value)>);
+			static_assert(offsetof(FPair, first) == offsetof(FKeyAndValue, Key));
+			static_assert(offsetof(FPair, second) == offsetof(FKeyAndValue, Value));
+			static_assert(sizeof(FPair) == sizeof(FKeyAndValue));
+			static_assert(alignof(FPair) == alignof(FKeyAndValue));
+			return *KeyAndValue;
+		}
+
+		TIterator& operator++()
+		{
+			++InnerIterator;
+			return *this;
+		}
+
+		bool operator == (const TIterator& Other) const
+		{
+			return InnerIterator == Other.InnerIterator;
+		}
+
+		bool operator != (const TIterator& Other) const
+		{
+			return InnerIterator != Other.InnerIterator;
+		}
+
+	private:
+		FInnerIterator InnerIterator;
+	};
+
+public:
+	using Key = KeyType;
+	using Value = ValueType;
+	using Iterator = TIterator</* bConst */ false>;
+	using ConstIterator = TIterator</* bConst */ true>;
+
+	// Constructor
+	THashMap() = default;
+	// Destructor
+	~THashMap() = default;
+	// Copy constructor
+	THashMap(const THashMap& Other) = default;
+	// Move constructor
+	THashMap(THashMap&& Other) : Map{std::move(Other.Map)}
+	{
+		if (this != &Other)
+		{
+			Other.Map.clear(); // Ensure Other is cleared.
+		}
+	}
+	// Copy assignment operator
+	THashMap& operator = (const THashMap& Other) = default;
+	// Move assignment operator
+	THashMap& operator = (THashMap&& Other)
+	{
+		if (this != &Other)
+		{
+			Map = std::move(Other.Map);
+			Other.Map.clear(); // Ensure Other is cleared.
+		}
+		return *this;
+	}
+
+	// Set the value associated with a key, replacing any existing entry with
+	// the given key.
+	template<typename K, typename V>
+	void Add(K&& Key, V&& Value)
+	{
+		Map[std::forward<K>(Key)] = std::forward<V>(Value);
+	}
+
+	// Looks up the value with the given key, returning a pointer to the value
+	// if found or nullptr if not found.
+	// Warning: The returned pointer will become invalid if the hash map is 
+	// modified.
+	template<typename K>
+	ValueType* Find(K&& Key)
+	{
+		typename FInnerMap::iterator It = Map.find(std::forward<K>(Key));
+		return It == Map.end() ? nullptr : &It->second;
+	}
+
+	// Looks up the value with the given key, returning a reference to the
+	// existing value if found or a reference to a newly added, zero-initialized
+	// value if not found.
+	// Warning: The returned reference will become invalid if the hash map is 
+	// modified.
+	template<typename K>
+	ValueType& FindOrAdd(K&& Key)
+	{
+		return Map[std::forward<K>(Key)];
+	}
+
+	// Removes the entry with the given key. This is a no-op if the hash map
+	// does not contain an entry with the given key.
+	void Remove(const KeyType& Key)
+	{
+		Map.erase(Key);
+	}
+
+	// Returns true if the hash map contains an entry with the given key.
+	bool Contains(const KeyType& Key)
+	{
+		return Map.count(Key) != 0;
+	}
+
+	// Removes all entries from the hash map, freeing all allocations.
+	void Empty()
+	{
+		Map.clear();
+	}
+
+	// Removes all entries from the hash map.
+	// TODO(SOL-7652): Preserve any heap allocations made by the hash map.
+	void Reset()
+	{
+		Map.clear();
+	}
+
+	// Returns the number of entries held by the hash map.
+	size_t Num() const
+	{
+		return Map.size();
+	}
+
+	// Returns true if the hash map holds no entries.
+	bool IsEmpty() const
+	{
+		return Map.empty();
+	}
+
+	// Iterator methods
+	ConstIterator begin() const { return ConstIterator{Map.begin()}; }
+	ConstIterator end() const { return ConstIterator{Map.end()}; }
+	Iterator begin() { return Iterator{Map.begin()}; }
+	Iterator end() { return Iterator{Map.end()}; }
+
+private:
+	FInnerMap Map;
+};
+
+}  // AutoRTFM
diff --git a/Engine/Source/Runtime/Core/Private/AutoRTFM/OpenHashThrottler.h b/Engine/Source/Runtime/Core/Private/AutoRTFM/OpenHashThrottler.h
--- a/Engine/Source/Runtime/Core/Private/AutoRTFM/OpenHashThrottler.h
+++ b/Engine/Source/Runtime/Core/Private/AutoRTFM/OpenHashThrottler.h
@@ -2,7 +2,7 @@
 
 #pragma once
 
-#include "Containers/Map.h"
+#include "HashMap.h"
 
 namespace AutoRTFM
 {
@@ -91,7 +91,7 @@ private:
 		// Total time spent hashing for all opens since last reset.
 		FSeconds TotalTimeSpentHashing = 0;
 		// A map of open return address to open info.
-		TMap<const void*, FOpenAddressInfo> Opens;
+		THashMap<const void*, FOpenAddressInfo> Opens;
 	};
 	
 	// Statistics for logging
diff --git a/Engine/Source/Runtime/Core/Private/AutoRTFM/StlAllocator.h b/Engine/Source/Runtime/Core/Private/AutoRTFM/StlAllocator.h
new file mode 100644
--- /dev/null
+++ b/Engine/Source/Runtime/Core/Private/AutoRTFM/StlAllocator.h
@@ -0,0 +1,49 @@
+// Copyright Epic Games, Inc. All Rights Reserved.
+
+#pragma once
+
+#include "AutoRTFM/ExternAPI.h"
+#include <cstddef>
+
+namespace AutoRTFM
+{
+
+// A stl-compatible allocator which allocates with AutoRTFM::Allocate() and
+// frees with AutoRTFM::Free().
+// TODO(SOL-7652): Remove this once HashMap has been replaced with a bespoke
+// hashmap implementation.
+template <typename T>
+class StlAllocator
+{
+public:
+	using value_type = T;
+
+	StlAllocator() = default;
+
+	template <typename U>
+	StlAllocator(const StlAllocator<U>&) {};
+
+	T* allocate(size_t Count)
+	{
+		return static_cast<T*>(AutoRTFM::Allocate(Count * sizeof(T), alignof(T)));
+	}
+
+	void deallocate(T* Pointer, size_t)
+	{
+		AutoRTFM::Free(Pointer);
+	}
+};
+
+template <typename T, typename U>
+bool operator==(const StlAllocator<T>&, const StlAllocator<U>&)
+{
+	return true;
+}
+
+template <typename T, typename U>
+bool operator!=(const StlAllocator<T>&, const StlAllocator<U>&)
+{
+	return false;
+}
+
+}  // AutoRTFM
diff --git a/Engine/Source/Runtime/Core/Private/AutoRTFM/TaskArray.h b/Engine/Source/Runtime/Core/Private/AutoRTFM/TaskArray.h
--- a/Engine/Source/Runtime/Core/Private/AutoRTFM/TaskArray.h
+++ b/Engine/Source/Runtime/Core/Private/AutoRTFM/TaskArray.h
@@ -2,10 +2,9 @@
 
 #pragma once
 
-#include "AutoRTFM/Utils.h"
+#include "HashMap.h"
 #include "Pool.h"
-#include "Containers/Map.h"
-#include "Templates/UnrealTemplate.h"
+#include "Utils.h"
 #include <cstddef>
 
 namespace AutoRTFM
@@ -137,8 +136,8 @@ public:
 		// Append the keyed entries of OtherArray to this.
 		for (auto& Iterator : OtherArray.Keyed)
 		{
-			KeyType Key = Iterator.template Get<0>();
-			TEntryStack& EntryStack = Iterator.template Get<1>();
+			KeyType Key = Iterator.Key;
+			TEntryStack& EntryStack = Iterator.Value;
 			Keyed.FindOrAdd(Key).Append(MoveTemp(EntryStack));
 		}
 
@@ -256,7 +255,7 @@ private:
 	FEntry* Tail = nullptr;
 
 	// A map of keyed tasks keys to their ordered task entries in the linked list.
-	TMap<KeyType, TEntryStack> Keyed;
+	THashMap<KeyType, TEntryStack> Keyed;
 
 	// Total number of tasks held by this array.
 	size_t Count = 0;
