[AutoRTFM] Fix accidental misreported refcounts in FThreadSafeRefCountedObject.

This class was rewritten at https://p4-swarm.epicgames.net/reviews/39426405. This accidentally introduced an issue where Release() would return the refcount _before_ release, instead of after release. We didn't have any tests which verified the behavior of these functions, so this caused a bug in the editor.

I've added tests to verify that AddRef, GetRefCount, and Release all return the expected refcounts. In non-transactional code, I've verified that the refcounts are accurate. In transactional code, I've verified that the returned value is the expected value or higher (which is all that we can promise).

diff --git a/Engine/Source/Programs/AutoRTFMTests/Private/RefCountPtrTests.cpp b/Engine/Source/Programs/AutoRTFMTests/Private/RefCountPtrTests.cpp
--- a/Engine/Source/Programs/AutoRTFMTests/Private/RefCountPtrTests.cpp
+++ b/Engine/Source/Programs/AutoRTFMTests/Private/RefCountPtrTests.cpp
@@ -2,6 +2,7 @@
 
 #include "Catch2Includes.h"
 #include <AutoRTFM/AutoRTFM.h>
+#include "AutoRTFM/AutoRTFMTesting.h"
 #include "Templates/RefCounting.h"
 
 namespace
@@ -25,13 +26,61 @@ namespace
 	};
 }
 
+TEMPLATE_TEST_CASE("CheckRefCounts", "[RefCountPtr]", RefCountedInt, RefCountedMixinInt, ThreadSafeRefCountedInt)
+{
+	// Refcounts should be correct outside of a transaction.
+	SECTION("Non-transactional refcounts should be accurate")
+	{
+		TestType* Ptr = new TestType(42);
+
+		REQUIRE(Ptr->GetRefCount() == 0);
+		       (Ptr->AddRef().CheckAtLeast(1));
+		REQUIRE(Ptr->GetRefCount() == 1);
+		       (Ptr->AddRef().CheckAtLeast(2));
+		REQUIRE(Ptr->GetRefCount() == 2);
+		REQUIRE(Ptr->Release() == 1);
+		REQUIRE(Ptr->GetRefCount() == 1);
+		       (Ptr->AddRef().CheckAtLeast(2));
+		REQUIRE(Ptr->GetRefCount() == 2);
+		REQUIRE(Ptr->Release() == 1);
+		REQUIRE(Ptr->GetRefCount() == 1);
+		REQUIRE(Ptr->Release() == 0);
+
+		// Ptr is dead
+	}
+
+	// Refcounts can be inflated inside a transaction, since releases are deferred.
+	SECTION("Transactional refcounts should represent a lower bound")
+	{
+		AutoRTFM::Testing::Commit([]
+		{
+			TestType* Ptr = new TestType(42);
+
+			REQUIRE(Ptr->GetRefCount() >= 0);
+				   (Ptr->AddRef().CheckAtLeast(1));
+			REQUIRE(Ptr->GetRefCount() >= 1);
+				   (Ptr->AddRef().CheckAtLeast(2));
+			REQUIRE(Ptr->GetRefCount() >= 2);
+			REQUIRE(Ptr->Release() >= 1);
+			REQUIRE(Ptr->GetRefCount() >= 1);
+				   (Ptr->AddRef().CheckAtLeast(2));
+			REQUIRE(Ptr->GetRefCount() >= 2);
+			REQUIRE(Ptr->Release() >= 1);
+			REQUIRE(Ptr->GetRefCount() >= 1);
+			REQUIRE(Ptr->Release() >= 0);
+
+			// Ptr is dead
+		});
+	}
+}
+
 TEMPLATE_TEST_CASE("PreviouslyAllocated", "[RefCountPtr]", RefCountedInt, RefCountedMixinInt, ThreadSafeRefCountedInt)
 {
 	SECTION("Using new T()")
 	{
 		TRefCountPtr<TestType> Foo(new TestType(42));
 
-		AutoRTFM::Commit([&]
+		AutoRTFM::Testing::Commit([&]
 		{
 			// Make a copy to bump the reference count.
 			TRefCountPtr<TestType> Copy = Foo;
@@ -46,7 +95,7 @@ TEMPLATE_TEST_CASE("PreviouslyAllocated", "[RefCountPtr]", RefCountedInt, RefCou
 	{
 		TRefCountPtr<TestType> Foo = MakeRefCount<TestType>(42);
 
-		AutoRTFM::Commit([&]
+		AutoRTFM::Testing::Commit([&]
 		{
 			// Make a copy to bump the reference count.
 			TRefCountPtr<TestType> Copy = Foo;
@@ -64,14 +113,14 @@ TEMPLATE_TEST_CASE("AbortWithPreviouslyAllocated", "[RefCountPtr]", RefCountedIn
 	{
 		TRefCountPtr<TestType> Foo(new TestType(42));
 
-		REQUIRE(AutoRTFM::ETransactionResult::AbortedByRequest == AutoRTFM::Transact([&]
+		AutoRTFM::Testing::Abort([&]
 		{
 			// Make a copy to bump the reference count.
 			TRefCountPtr<TestType> Copy = Foo;
 			Copy->Value = 13;
 
 			AutoRTFM::AbortTransaction();
-		}));
+		});
 
 		REQUIRE(Foo->Value == 42);
 		REQUIRE(Foo.GetRefCount() == 1);
@@ -81,14 +130,14 @@ TEMPLATE_TEST_CASE("AbortWithPreviouslyAllocated", "[RefCountPtr]", RefCountedIn
 	{
 		TRefCountPtr<TestType> Foo = MakeRefCount<TestType>(42);
 
-		REQUIRE(AutoRTFM::ETransactionResult::AbortedByRequest == AutoRTFM::Transact([&]
+		AutoRTFM::Testing::Abort([&]
 		{
 			// Make a copy to bump the reference count.
 			TRefCountPtr<TestType> Copy = Foo;
 			Copy->Value = 13;
 
 			AutoRTFM::AbortTransaction();
-		}));
+		});
 
 		REQUIRE(Foo->Value == 42);
 		REQUIRE(Foo.GetRefCount() == 1);
@@ -101,44 +150,44 @@ TEMPLATE_TEST_CASE("AbortWithNewlyAllocated", "[RefCountPtr]", RefCountedInt, Re
 
 	SECTION("Using new T()")
 	{
-		REQUIRE(AutoRTFM::ETransactionResult::AbortedByRequest == AutoRTFM::Transact([&]
+		AutoRTFM::Testing::Abort([&]
 		{
 			TRefCountPtr<TestType> Foo(new TestType(42));
 			Result = Foo->Value;
 			AutoRTFM::AbortTransaction();
-		}));
+		});
 	}
 
 	SECTION("Using new T() and a copy")
 	{
-		REQUIRE(AutoRTFM::ETransactionResult::AbortedByRequest == AutoRTFM::Transact([&]
+		AutoRTFM::Testing::Abort([&]
 		{
 			TRefCountPtr<TestType> Foo(new TestType(42));
 			TRefCountPtr<TestType> Copy = Foo;
 			Result = Copy->Value;
 			AutoRTFM::AbortTransaction();
-		}));
+		});
 	}
 
 	SECTION("Using MakeRefCount<T>")
 	{
-		REQUIRE(AutoRTFM::ETransactionResult::AbortedByRequest == AutoRTFM::Transact([&]
+		AutoRTFM::Testing::Abort([&]
 		{
 			TRefCountPtr<TestType> Foo = MakeRefCount<TestType>(42);
 			Result = Foo->Value;
 			AutoRTFM::AbortTransaction();
-		}));
+		});
 	}
 
 	SECTION("Using MakeRefCount<T> and a copy")
 	{
-		REQUIRE(AutoRTFM::ETransactionResult::AbortedByRequest == AutoRTFM::Transact([&]
+		AutoRTFM::Testing::Abort([&]
 		{
 			TRefCountPtr<TestType> Foo = MakeRefCount<TestType>(42);
 			TRefCountPtr<TestType> Copy = Foo;
 			Result = Copy->Value;
 			AutoRTFM::AbortTransaction();
-		}));
+		});
 	}
 
 	REQUIRE(Result == 42);
@@ -150,7 +199,7 @@ TEMPLATE_TEST_CASE("OnCommitCapturingRefCountPtr", "[RefCountPtr]", RefCountedIn
 
 	SECTION("Committing")
 	{
-		AutoRTFM::Commit([&]
+		AutoRTFM::Testing::Commit([&]
 		{
 			AutoRTFM::OnCommit([Foo]
 			{
@@ -163,14 +212,14 @@ TEMPLATE_TEST_CASE("OnCommitCapturingRefCountPtr", "[RefCountPtr]", RefCountedIn
 
 	SECTION("Aborting")
 	{
-		REQUIRE(AutoRTFM::ETransactionResult::AbortedByRequest == AutoRTFM::Transact([&]
+		AutoRTFM::Testing::Abort([&]
 		{
 			AutoRTFM::OnCommit([Foo]
 			{
 				Foo->Value = 13;
 			});
 			AutoRTFM::AbortTransaction();
-		}));
+		});
 
 		REQUIRE(Foo->Value == 42);
 	}
@@ -182,7 +231,7 @@ TEMPLATE_TEST_CASE("OnAbortCapturingRefCountPtr", "[RefCountPtr]", RefCountedInt
 
 	SECTION("Committing")
 	{
-		AutoRTFM::Commit([&]
+		AutoRTFM::Testing::Commit([&]
 		{
 			AutoRTFM::OnAbort([Foo]
 			{
@@ -196,14 +245,14 @@ TEMPLATE_TEST_CASE("OnAbortCapturingRefCountPtr", "[RefCountPtr]", RefCountedInt
 
 	SECTION("Aborting")
 	{
-		REQUIRE(AutoRTFM::ETransactionResult::AbortedByRequest == AutoRTFM::Transact([&]
+		AutoRTFM::Testing::Abort([&]
 		{
 			AutoRTFM::OnAbort([Foo]
 			{
 				Foo->Value = 13;
 			});
 			AutoRTFM::AbortTransaction();
-		}));
+		});
 
 		REQUIRE(Foo->Value == 13);
 	}
diff --git a/Engine/Source/Runtime/Core/Public/Templates/RefCounting.h b/Engine/Source/Runtime/Core/Public/Templates/RefCounting.h
--- a/Engine/Source/Runtime/Core/Public/Templates/RefCounting.h
+++ b/Engine/Source/Runtime/Core/Public/Templates/RefCounting.h
@@ -139,7 +139,7 @@ public:
 		}
 		else
 		{
-			RefsToReturn = ImmediatelyRelease<DeleteFn>();
+			RefsToReturn = ImmediatelyRelease<DeleteFn>() - 1;
 		}
 
 		return RefsToReturn;
@@ -373,10 +373,7 @@ public:
 
 	uint32 Release() const
 	{
-		Super::template Release<StaticDestroyMixin>();
-
-		// Note: TRefCountPtr doesn't use the return value
-		return 0;
+		return Super::template Release<StaticDestroyMixin>();
 	}
 
 	uint32 GetRefCount() const
