[AutoRTFM] Replace use of UE's TBitArray with a new AutoRTFM::TBitStack

Part of the work to separate AutoRTFM out of core to break cyclic build dependency issues.

diff --git a/Engine/Source/Programs/AutoRTFMTests/Private/BitStackTests.cpp b/Engine/Source/Programs/AutoRTFMTests/Private/BitStackTests.cpp
new file mode 100644
--- /dev/null
+++ b/Engine/Source/Programs/AutoRTFMTests/Private/BitStackTests.cpp
@@ -0,0 +1,382 @@
+// Copyright Epic Games, Inc. All Rights Reserved.
+
+#include "AutoRTFM/BitStack.h"
+#include "Catch2Includes.h"
+
+TEST_CASE("BitStack")
+{
+	using BitStackType = AutoRTFM::TBitStack<8>;
+
+	auto Check = [&](const BitStackType& BitStack, std::vector<bool> Expected)
+	{
+		// Check the reported count is as expected.
+		REQUIRE(BitStack.Num() == Expected.size());
+		REQUIRE(BitStack.IsEmpty() == Expected.empty());
+
+		for (size_t I = 0; I < Expected.size(); I++)
+		{
+			// Check the bit is as expected.
+			REQUIRE(BitStack[I] == Expected[I]);
+			REQUIRE(const_cast<const BitStackType&>(BitStack)[I] == Expected[I]);
+		}
+	};
+
+	SECTION("Push, Pop")
+	{
+		BitStackType BitStack;
+		Check(BitStack, {});
+
+		BitStack.Push(true);
+		Check(BitStack, {true});
+
+		BitStack.Push(false);
+		Check(BitStack, {true, false});
+
+		BitStack.Push(false);
+		Check(BitStack, {true, false, false});
+
+		REQUIRE(BitStack.Pop() == false);
+		Check(BitStack, {true, false});
+
+		BitStack.Push(true);
+		Check(BitStack, {true, false, true});
+
+		BitStack.Push(true);
+		Check(BitStack, {true, false, true, true});
+
+		BitStack.Push(true); // Spill inline -> heap
+		Check(BitStack, {true, false, true, true, true});
+
+		REQUIRE(BitStack.Pop() == true);
+		Check(BitStack, {true, false, true, true});
+	}
+
+	SECTION("Index operator")
+	{
+		BitStackType BitStack;
+		BitStack.Push(true);
+		BitStack.Push(true);
+		BitStack.Push(false);
+		BitStack.Push(true);
+		BitStack.Push(true);
+		Check(BitStack, {true, true, false, true, true});
+		
+		BitStack[0] = false;
+		Check(BitStack, {false, true, false, true, true});
+
+		BitStack[1] = false;
+		Check(BitStack, {false, false, false, true, true});
+
+		BitStack[2] = true;
+		Check(BitStack, {false, false, true, true, true});
+
+		BitStack[3] = BitStack[0];
+		Check(BitStack, {false, false, true, false, true});
+
+		BitStack[0] = BitStack[4];
+		Check(BitStack, {true, false, true, false, true});
+	}
+
+	SECTION("Clear / Reset")
+	{
+		BitStackType BitStack;
+		BitStack.Push(true);
+		BitStack.Push(false);
+		BitStack.Push(false);
+		BitStack.Push(true);
+		BitStack.Push(false);
+		SECTION("Clear")
+		{
+			BitStack.Clear();
+			Check(BitStack, {});
+			BitStack.Push(true);
+			BitStack.Push(false);
+			Check(BitStack, {true, false});
+		}
+		SECTION("Reset")
+		{
+			BitStack.Reset();
+			Check(BitStack, {});
+			BitStack.Push(true);
+			BitStack.Push(false);
+			Check(BitStack, {true, false});
+		}
+	}
+
+	SECTION("Copy Construct")
+	{
+		SECTION("Inline")
+		{
+			BitStackType Source;
+			Source.Push(true);
+			Source.Push(false);
+			Source.Push(true);
+			BitStackType Target(Source);
+			Check(Source, {true, false, true});
+			Check(Target, {true, false, true});
+		}
+
+		SECTION("Heap")
+		{
+			BitStackType Source;
+			Source.Push(true);
+			Source.Push(false);
+			Source.Push(false);
+			Source.Push(true);
+			Source.Push(true);
+			BitStackType Target(Source);
+			Check(Source, {true, false, false, true, true});
+			Check(Target, {true, false, false, true, true});
+		}
+	}
+
+	SECTION("Move Construct")
+	{
+		SECTION("Inline")
+		{
+			BitStackType Source;
+			Source.Push(true);
+			Source.Push(true);
+			Source.Push(false);
+			BitStackType Target(std::move(Source));
+			Check(Source, {});
+			Check(Target, {true, true, false});
+		}
+
+		SECTION("Heap")
+		{
+			BitStackType Source;
+			Source.Push(false);
+			Source.Push(true);
+			Source.Push(false);
+			Source.Push(false);
+			Source.Push(true);
+			BitStackType Target(std::move(Source));
+			Check(Source, {});
+			Check(Target, {false, true, false, false, true});
+		}
+	}
+
+	SECTION("Copy Assign")
+	{
+		BitStackType Source;
+		BitStackType Target;
+		auto CheckCopy = [&]()
+		{
+			SECTION("Source Empty")
+			{
+				Target = Source;
+				Check(Source, {});
+				Check(Target, {});
+			}
+			SECTION("Source Inline")
+			{
+				Source.Push(true);
+				Source.Push(false);
+				Source.Push(true);
+				Target = Source;
+				Check(Source, {true, false, true});
+				Check(Target, {true, false, true});
+			}
+			SECTION("Source Heap")
+			{
+				Source.Push(false);
+				Source.Push(false);
+				Source.Push(true);
+				Source.Push(false);
+				Source.Push(true);
+				Target = Source;
+				Check(Source, {false, false, true, false, true});
+				Check(Target, {false, false, true, false, true});
+			}
+		};
+
+		SECTION("Target Empty")
+		{
+			CheckCopy();
+		}
+		SECTION("Target Inline")
+		{
+			Target.Push(true);
+			Target.Push(true);
+			Target.Push(false);
+			CheckCopy();
+		}
+		SECTION("Target Heap")
+		{
+			Target.Push(false);
+			Target.Push(true);
+			Target.Push(false);
+			Target.Push(true);
+			Target.Push(false);
+			CheckCopy();
+		}
+	}
+
+	SECTION("Copy Assign Self")
+	{
+		BitStackType BitStack;
+		BitStackType& Alias = BitStack;
+		SECTION("Empty")
+		{
+			BitStack = Alias;
+			Check(BitStack, {});
+		}
+		SECTION("Inline")
+		{
+			BitStack.Push(false);
+			BitStack.Push(false);
+			BitStack.Push(true);
+			BitStack = Alias;
+			Check(BitStack, {false, false, true});
+		}
+		SECTION("Heap")
+		{
+			BitStack.Push(true);
+			BitStack.Push(true);
+			BitStack.Push(false);
+			BitStack.Push(false);
+			BitStack.Push(true);
+			BitStack = Alias;
+			Check(BitStack, {true, true, false, false, true});
+		}
+	}
+
+	SECTION("Move Assign")
+	{
+		BitStackType Source;
+		BitStackType Target;
+		auto CheckMove = [&]
+		{
+			SECTION("Source Empty")
+			{
+				Target = std::move(Source);
+				Check(Source, {});
+				Check(Target, {});
+			}
+			SECTION("Source Inline")
+			{
+				Source.Push(true);
+				Source.Push(true);
+				Source.Push(false);
+				Target = std::move(Source);
+				Check(Source, {});
+				Check(Target, {true, true, false});
+			}
+			SECTION("Source Heap")
+			{
+				Source.Push(true);
+				Source.Push(true);
+				Source.Push(false);
+				Source.Push(true);
+				Source.Push(true);
+				Target = std::move(Source);
+				Check(Source, {});
+				Check(Target, {true, true, false, true, true});
+			}
+		};
+
+		SECTION("Target Empty")
+		{
+			CheckMove();
+		}
+		SECTION("Target Inline")
+		{
+			Target.Push(true);
+			Target.Push(false);
+			Target.Push(true);
+			CheckMove();
+		}
+		SECTION("Target Heap")
+		{
+			Target.Push(true);
+			Target.Push(false);
+			Target.Push(false);
+			Target.Push(false);
+			Target.Push(true);
+			CheckMove();
+		}
+	}
+
+	SECTION("Move Assign Self")
+	{
+		BitStackType BitStack;
+		BitStackType& Alias = BitStack;
+		SECTION("Empty")
+		{
+			BitStack = std::move(Alias);
+			Check(BitStack, {});
+		}
+		SECTION("Inline")
+		{
+			BitStack.Push(false);
+			BitStack.Push(true);
+			BitStack.Push(true);
+			BitStack = std::move(Alias);
+			Check(BitStack, {false, true, true});
+		}
+		SECTION("Heap")
+		{
+			BitStack.Push(true);
+			BitStack.Push(false);
+			BitStack.Push(true);
+			BitStack.Push(true);
+			BitStack.Push(false);
+			BitStack = std::move(Alias);
+			Check(BitStack, {true, false, true, true, false});
+		}
+	}
+
+	SECTION("Soak")
+	{
+		BitStackType BitStack;
+		std::vector<bool> Expected;
+		for (int I = 0; I < 10000; I++)
+		{
+			int Mod100 = (I * 15485863) % 100;
+			switch (Mod100)
+			{
+				case 0:
+				{
+					BitStack.Clear();
+					Expected.clear();
+					break;
+				}
+				case 1:
+				{	
+					BitStack.Reset();
+					Expected.clear();
+					break;
+				}
+				case 2:
+				{	
+					BitStackType Copy(BitStack);
+					BitStack = Copy;
+					break;
+				}
+				case 3:
+				{	
+					BitStackType Move(std::move(BitStack));
+					BitStack = std::move(Move);
+					break;
+				}
+				default:
+				{
+					if (Mod100 > 40 || Expected.empty())
+					{
+						BitStack.Push((I & 1) == 0);
+						Expected.push_back((I & 1) == 0);
+						Check(BitStack, Expected);
+					}
+					else
+					{
+						REQUIRE(BitStack.Pop() == Expected.back());
+						Expected.pop_back();
+						Check(BitStack, Expected);
+					}
+					break;
+				}
+			}
+		}
+	}
+}
diff --git a/Engine/Source/Runtime/Core/Private/AutoRTFM/BitStack.h b/Engine/Source/Runtime/Core/Private/AutoRTFM/BitStack.h
new file mode 100644
--- /dev/null
+++ b/Engine/Source/Runtime/Core/Private/AutoRTFM/BitStack.h
@@ -0,0 +1,188 @@
+// Copyright Epic Games, Inc. All Rights Reserved.
+
+#pragma once
+
+#include "ContainerValidation.h"
+#include "Stack.h"
+#include "Utils.h"
+#include <cstddef>
+#include <cstdint>
+
+namespace AutoRTFM
+{
+
+// A dynamically sized stack of bits.
+//
+// Notes:
+//  * Heap allocated memory is not automatically freed when popping elements.
+//    Only calling Reset() or destructing the stack will free heap allocated
+//    memory.
+//  * This class is not relocatable and so is not safe to use in UE's containers
+//    which require elements to be relocatable.
+//
+// Template parameters:
+//   InlineCapacity - the number of elements that can be held before spilling to
+//                    the heap.
+//   Validation - if Disabled, do not perform validity assertions.
+template<size_t InlineCapacity, EContainerValidation Validation = EContainerValidation::Enabled>
+class TBitStack
+{
+	using WordType = uint64_t;
+	static constexpr size_t NumWordBits = 64;
+	static_assert(NumWordBits == sizeof(WordType) * 8);
+
+public:
+	// A reference to a single bit in the stack.
+	// Becomes invalid if the stack holding the bit is modified.
+	class FBitRef
+	{
+	public:
+		// Constructor
+		FBitRef(WordType& Word, WordType Mask) : Word(Word), Mask(Mask) {}
+
+		// Returns true if the bit is set.
+		operator bool() const
+		{
+			return (Word & Mask) != 0;
+		}
+
+		// Copy assignment operator.
+		// Changes the value of the referenced bit in the stack.
+		FBitRef& operator = (const FBitRef& Other)
+		{
+			return *this = static_cast<bool>(Other);
+		}
+
+		// Assignment operator.
+		// Changes the value of the referenced bit in the stack.
+		FBitRef& operator = (bool Value)
+		{
+			Word = Value ? (Word | Mask) : (Word & ~Mask);
+			return *this;
+		}
+
+	private:
+		WordType& Word;
+		const WordType Mask;
+	};
+
+	// Constructor.
+	TBitStack() = default;
+	// Copy constructor.
+	TBitStack(const TBitStack& Other) = default;
+	// Move constructor.
+	TBitStack(TBitStack&& Other) : Words{std::move(Other.Words)}, Count{Other.Count}
+	{
+		if (&Other != this)
+		{
+			Other.Count = 0;
+		}
+	}
+	// Destructor.
+	~TBitStack() = default;
+	// Copy assignment operator.
+	TBitStack& operator=(const TBitStack& Other) = default;
+	// Move assignment operator.
+	TBitStack& operator=(TBitStack&& Other)
+	{
+		if (&Other != this)
+		{
+			Words = std::move(Other.Words);
+			Count = Other.Count;
+			Other.Count = 0;
+		}
+		return *this;
+	}
+	
+	// Clears all the items from the stack, preserving the capacity.
+	void Clear()
+	{
+		Words.Clear();
+		Count = 0;
+	}
+
+	// Clears all the items from the stack, freeing all heap allocations and
+	// resetting the capacity to InlineCapacity
+	void Reset()
+	{
+		Words.Reset();
+		Count = 0;
+	}
+
+	// Pushes a new bit on to the stack.
+	inline void Push(bool Bit)
+	{
+		const WordType WordBitIndex = Count & (NumWordBits - 1);
+		if (WordBitIndex == 0)
+		{
+			Words.Push(Bit ? 1 : 0);
+		}
+		else
+		{
+			WordType& Word = Words.Back();
+			WordType Mask = static_cast<WordType>(1) << WordBitIndex;
+			Word = Bit ? (Word | Mask) : (Word & ~Mask);
+		}
+		++Count;
+		ASSERT(Validation == EContainerValidation::Disabled || NumWordsFor(Count) == Words.Num());
+	}
+
+	// Removes the last item on the stack.
+	inline bool Pop()
+	{
+		ASSERT(Validation == EContainerValidation::Disabled || Count > 0);
+		--Count;
+		const WordType WordBitIndex = Count & (NumWordBits - 1);
+		const WordType Word = Words.Back();
+		const bool Value = ((Word >> WordBitIndex) & 1) != 0;
+		if (WordBitIndex == 0)
+		{
+			Words.Pop();
+		}
+		ASSERT(Validation == EContainerValidation::Disabled || NumWordsFor(Count) == Words.Num());
+		return Value;
+	}
+
+	// Reserves memory for NewCapacity bits.
+	inline void Reserve(size_t NewCapacity)
+	{
+		Words.Reserve(NumWordsFor(NewCapacity));
+	}
+
+	// Returns the number of bits held by the stack.
+	inline size_t Num() const { return Count; }
+
+	// Returns true if the stack holds no bits.
+	inline bool IsEmpty() const { return Count == 0; }
+
+	// Index operator.
+	// Returns a reference to the bit with the given index in the stack.
+	FBitRef operator[](size_t Index)
+	{
+		ASSERT(Validation == EContainerValidation::Disabled || Index < Count);
+		WordType& Word = Words[Index / NumWordBits];
+		WordType Mask = static_cast<WordType>(1) << (Index & (NumWordBits-1));
+		return FBitRef{Word, Mask};
+	}
+
+	// Constant index operator.
+	// Returns true if the bit with the given index is set.
+	bool operator[](size_t Index) const
+	{
+		ASSERT(Validation == EContainerValidation::Disabled || Index < Count);
+		WordType Word = Words[Index / NumWordBits];
+		WordType Mask = static_cast<WordType>(1) << (Index & (NumWordBits-1));
+		return (Word & Mask) != 0;
+	}
+
+private:
+	static constexpr size_t NumWordsFor(size_t NumBits)
+	{
+		return (NumBits + (NumWordBits-1)) / NumWordBits;
+	}
+
+	TStack<WordType, NumWordsFor(InlineCapacity)> Words;
+	size_t Count = 0; // In bits
+};
+
+}
diff --git a/Engine/Source/Runtime/Core/Private/AutoRTFM/ContainerValidation.h b/Engine/Source/Runtime/Core/Private/AutoRTFM/ContainerValidation.h
new file mode 100644
--- /dev/null
+++ b/Engine/Source/Runtime/Core/Private/AutoRTFM/ContainerValidation.h
@@ -0,0 +1,16 @@
+// Copyright Epic Games, Inc. All Rights Reserved.
+
+#pragma once
+
+namespace AutoRTFM
+{
+
+// Enumerator used to control validity assertions for containers
+// (bounds checking, etc).
+enum class EContainerValidation
+{
+	Enabled,
+	Disabled
+};
+
+}
diff --git a/Engine/Source/Runtime/Core/Private/AutoRTFM/IntervalTree.h b/Engine/Source/Runtime/Core/Private/AutoRTFM/IntervalTree.h
--- a/Engine/Source/Runtime/Core/Private/AutoRTFM/IntervalTree.h
+++ b/Engine/Source/Runtime/Core/Private/AutoRTFM/IntervalTree.h
@@ -2,7 +2,8 @@
 
 #pragma once
 
-#include "Containers/BitArray.h"
+#include "BitStack.h"
+#include "ContainerValidation.h"
 #include "Stack.h"
 #include "Utils.h"
 
@@ -45,7 +46,7 @@ struct FIntervalTree final
                 return true;
             }
 
-			const ArrayType<FIntervalTreeNodeIndex>& Next = (NewRange.Start < Range.End) ? NodeLefts : NodeRights;
+			const StackType<FIntervalTreeNodeIndex>& Next = (NewRange.Start < Range.End) ? NodeLefts : NodeRights;
 			Current = Next[Current];
         } while (IntervalTreeNodeIndexNone != Current);
 
@@ -74,7 +75,7 @@ struct FIntervalTree final
             return;
         }
 
-		ArrayType<FIntervalTreeNodeIndex> ToProcess;
+		StackType<FIntervalTreeNodeIndex> ToProcess;
 
         ToProcess.Push(Other.Root);
 
@@ -103,7 +104,9 @@ struct FIntervalTree final
     }
 
 private:
-	template<typename T> using ArrayType = TStack<T, 8, EStackValidation::Disabled>;
+	static constexpr size_t InlineArraySize = 8;
+	template<typename T> using StackType = TStack<T, InlineArraySize, EContainerValidation::Disabled>;
+	using BitStackType = TBitStack<InlineArraySize, EContainerValidation::Disabled>;
 
     struct FRange final
     {
@@ -125,11 +128,11 @@ private:
 	};
 
     FIntervalTreeNodeIndex Root = IntervalTreeNodeIndexNone;
-	ArrayType<FRange> NodeRanges;
-	ArrayType<FIntervalTreeNodeIndex> NodeLefts;
-	ArrayType<FIntervalTreeNodeIndex> NodeRights;
-	ArrayType<FIntervalTreeNodeIndex> NodeParents;
-	TBitArray<> NodeColors;
+	StackType<FRange> NodeRanges;
+	StackType<FIntervalTreeNodeIndex> NodeLefts;
+	StackType<FIntervalTreeNodeIndex> NodeRights;
+	StackType<FIntervalTreeNodeIndex> NodeParents;
+	BitStackType NodeColors;
 
 	static constexpr bool bExtraDebugging = false;
 
@@ -140,7 +143,7 @@ private:
 		NodeLefts.Push(IntervalTreeNodeIndexNone);
 		NodeRights.Push(IntervalTreeNodeIndexNone);
 		NodeParents.Push(IntervalTreeNodeIndexNone);
-		NodeColors.Add(EColor::Black);
+		NodeColors.Push(EColor::Black);
 		return Index;
 	}
 
@@ -151,7 +154,7 @@ private:
 		NodeLefts.Push(IntervalTreeNodeIndexNone);
 		NodeRights.Push(IntervalTreeNodeIndexNone);
 		NodeParents.Push(Parent);
-		NodeColors.Add(EColor::Red);
+		NodeColors.Push(EColor::Red);
 		return Index;
 	}
 
diff --git a/Engine/Source/Runtime/Core/Private/AutoRTFM/Stack.h b/Engine/Source/Runtime/Core/Private/AutoRTFM/Stack.h
--- a/Engine/Source/Runtime/Core/Private/AutoRTFM/Stack.h
+++ b/Engine/Source/Runtime/Core/Private/AutoRTFM/Stack.h
@@ -4,20 +4,13 @@
 
 #include "ExternAPI.h"
 #include "Utils.h"
+#include "ContainerValidation.h"
 #include <cstddef>
 #include <type_traits>
 
 namespace AutoRTFM
 {
 
-// Enumerator used to control validity assertions for TStack (bounds checking, 
-// etc).
-enum class EStackValidation
-{
-	Enabled,
-	Disabled
-};
-
 // A stack container with an inline fixed capacity that once exceeded spills to
 // the heap.
 //
@@ -33,7 +26,7 @@ enum class EStackValidation
 //   InlineCapacity - the number of elements that can be held before spilling to
 //                    the heap.
 //   Validation - if Disabled, do not perform validity assertions.
-template<typename T, size_t InlineCapacity, EStackValidation Validation = EStackValidation::Enabled>
+template<typename T, size_t InlineCapacity, EContainerValidation Validation = EContainerValidation::Enabled>
 class TStack
 {
 public:
@@ -83,7 +76,7 @@ public:
 		return *this;
 	}
 	
-	// Clears all the items from the vector, preserving the capacity.
+	// Clears all the items from the stack, preserving the capacity.
 	void Clear()
 	{
 		for (T& Item : *this)
@@ -93,7 +86,7 @@ public:
 		Count = 0;
 	}
 
-	// Clears all the items from the vector, freeing all heap allocations and
+	// Clears all the items from the stack, freeing all heap allocations and
 	// resetting the capacity to InlineCapacity
 	void Reset()
 	{
@@ -147,7 +140,7 @@ public:
 	// Removes the last item on the stack.
 	inline void Pop()
 	{
-		ASSERT(Validation == EStackValidation::Disabled || Count > 0);
+		ASSERT(Validation == EContainerValidation::Disabled || Count > 0);
 		--Count;
 		Data[Count].~T();
 	}
@@ -197,37 +190,37 @@ public:
 
 	T& operator[](size_t Index)
 	{
-		ASSERT(Validation == EStackValidation::Disabled || Index < Count);
+		ASSERT(Validation == EContainerValidation::Disabled || Index < Count);
 		return Data[Index];
 	}
 
 	const T& operator[](size_t Index) const
 	{
-		ASSERT(Validation == EStackValidation::Disabled || Index < Count);
+		ASSERT(Validation == EContainerValidation::Disabled || Index < Count);
 		return Data[Index];
 	}
 
 	T& Front()
 	{
-		ASSERT(Validation == EStackValidation::Disabled || Count > 0);
+		ASSERT(Validation == EContainerValidation::Disabled || Count > 0);
 		return Data[0];
 	}
 
 	const T& Front() const
 	{
-		ASSERT(Validation == EStackValidation::Disabled || Count > 0);
+		ASSERT(Validation == EContainerValidation::Disabled || Count > 0);
 		return Data[0];
 	}
 
 	T& Back()
 	{
-		ASSERT(Validation == EStackValidation::Disabled || Count > 0);
+		ASSERT(Validation == EContainerValidation::Disabled || Count > 0);
 		return Data[Count - 1];
 	}
 
 	const T& Back() const
 	{
-		ASSERT(Validation == EStackValidation::Disabled || Count > 0);
+		ASSERT(Validation == EContainerValidation::Disabled || Count > 0);
 		return Data[Count - 1];
 	}
 
