[AutoRTFM] Add reproducer for SOL-7447 in AutoRTFMTests.

diff --git a/Engine/Source/Programs/AutoRTFMTests/Private/Statics.cpp b/Engine/Source/Programs/AutoRTFMTests/Private/Statics.cpp
--- a/Engine/Source/Programs/AutoRTFMTests/Private/Statics.cpp
+++ b/Engine/Source/Programs/AutoRTFMTests/Private/Statics.cpp
@@ -5,6 +5,7 @@
 #include "AutoRTFM/ContextInlines.h"
 #include "AutoRTFM/Transaction.h"
 #include "Catch2Includes.h"
+#include "CrossCUTests.h"
 #include <AutoRTFM/AutoRTFMTesting.h>
 #include <chrono>
 #include <condition_variable>
@@ -383,3 +384,111 @@ TEST_CASE("Static.Concurrent")
 	StaticInitFn();
 	TransactionalThread.join();
 }
+
+// The following code is a reduced version of the bug which would trigger SOL-7447.
+namespace 
+{
+namespace SOL7447
+{
+
+struct MyString
+{
+	MyString() = default;
+	~MyString() 
+	{ 
+		// This is only here so that the compiler won't optimize the destructor away.
+		// You can also reproduce by writing through a volatile or to a global variable.
+		CrossCU::SomeFunction(Char);
+	}
+
+	char Char = 0;
+};
+
+struct FFormattingRules
+{
+	char NegativePrefixString;
+	char NegativeSuffixString;
+	char PositivePrefixString;
+	char PositiveSuffixString;
+};
+
+class FSigningStrings
+{
+public:
+	FSigningStrings(const FFormattingRules& InFormattingRules)
+	{
+		// Resolve out the default cases
+		if (InFormattingRules.NegativePrefixString)
+		{
+			NegativePrefixStringPtr = &InFormattingRules.NegativePrefixString;
+		}
+		if (InFormattingRules.NegativeSuffixString)
+		{
+			NegativeSuffixStringPtr = &InFormattingRules.NegativeSuffixString;
+		}
+		if (InFormattingRules.PositivePrefixString)
+		{
+			PositivePrefixStringPtr = &InFormattingRules.PositivePrefixString;
+		}
+		if (InFormattingRules.PositiveSuffixString)
+		{
+			PositiveSuffixStringPtr = &InFormattingRules.PositiveSuffixString;
+		}
+	}
+
+	const char& GetNegativePrefixString() const
+	{
+		return NegativePrefixStringPtr ? *NegativePrefixStringPtr : GetEmptyString().Char;
+	}
+
+	const char& GetNegativeSuffixString() const
+	{
+		return NegativeSuffixStringPtr ? *NegativeSuffixStringPtr : GetEmptyString().Char;
+	}
+
+	const char& GetPositivePrefixString() const
+	{
+		return PositivePrefixStringPtr ? *PositivePrefixStringPtr : GetEmptyString().Char;
+	}
+
+	const char& GetPositiveSuffixString() const
+	{
+		return PositiveSuffixStringPtr ? *PositiveSuffixStringPtr : GetEmptyString().Char;
+	}
+
+private:
+	const char* NegativePrefixStringPtr = nullptr;
+	const char* NegativeSuffixStringPtr = nullptr;
+	const char* PositivePrefixStringPtr = nullptr;
+	const char* PositiveSuffixStringPtr = nullptr;
+
+	const MyString& GetEmptyString() const
+	{
+		static const MyString EmptyStr = MyString();
+		return EmptyStr;
+	}
+};
+
+TPair<char, char> BuildFinalString(const FFormattingRules& InFormattingRules, bool bIsNegative) 
+{
+	const FSigningStrings SigningStrings(InFormattingRules);
+	const char& FinalPrefixStr = bIsNegative ? SigningStrings.GetNegativePrefixString() : SigningStrings.GetPositivePrefixString();
+ 	const char& FinalSuffixStr = bIsNegative ? SigningStrings.GetNegativeSuffixString() : SigningStrings.GetPositiveSuffixString();
+
+	return {FinalPrefixStr, FinalSuffixStr};
+}
+
+} // namespace SOL7447
+} // namespace
+
+TEST_CASE("Static.ComplexStaticInitialization")
+{
+	AutoRTFM::Testing::Commit([]
+	{
+		SOL7447::FFormattingRules Rules{'+', 0, '-', 0};
+		TPair<char, char> A = SOL7447::BuildFinalString(Rules, false);
+		TPair<char, char> B = SOL7447::BuildFinalString(Rules, true);
+
+		REQUIRE(A != B);
+	});
+}
