[AutoRTFM] Replace use of FPlatformTLS::GetCurrentThreadId() with new AutoRTFM::FThreadID

Part of the work to separate AutoRTFM out of core to break cyclic build dependency issues.

diff --git a/Engine/Source/Programs/AutoRTFMTests/Private/ThreadIDTests.cpp b/Engine/Source/Programs/AutoRTFMTests/Private/ThreadIDTests.cpp
new file mode 100644
--- /dev/null
+++ b/Engine/Source/Programs/AutoRTFMTests/Private/ThreadIDTests.cpp
@@ -0,0 +1,88 @@
+// Copyright Epic Games, Inc. All Rights Reserved.
+
+#include "AutoRTFM/ThreadID.h"
+#include "Catch2Includes.h"
+
+#include <condition_variable>
+#include <cstddef>
+#include <mutex>
+#include <thread>
+#include <vector>
+#include <unordered_set>
+
+TEST_CASE("ThreadID")
+{
+	REQUIRE(AutoRTFM::FThreadID{} == AutoRTFM::FThreadID::Invalid);
+	
+	// Synchronisation primitive used to ensure all threads have obtained
+	// their ID before continuing on the main thread.
+	class FWaitGroup
+	{
+	public:
+		// Constructor. Creates the WaitGroup with an the initial count.
+		FWaitGroup(size_t Count) : Count{Count}{}
+
+		// Blocks until Done() is called Count times.
+		void Wait()
+		{
+			std::unique_lock<std::mutex> Lock{Mutex};
+			ConditionVariable.wait(Lock, [this] { return Count == 0; });
+		}
+
+		// Decrements Count, unblocking any Wait() calls when Count reaches 0.
+		void Done()
+		{
+			std::unique_lock<std::mutex> Lock{Mutex};
+			Count--;
+			if (Count == 0)
+			{
+				ConditionVariable.notify_all();
+			}
+		}
+
+	private:
+		std::mutex Mutex;
+		std::condition_variable ConditionVariable;
+		size_t Count = 0;
+	};
+
+	// A STL hasher for AutoRTFM::FThreadID so we can use them in a std::unordered_map.
+	struct FThreadIDHasher
+	{
+		size_t operator()(const AutoRTFM::FThreadID& ThreadID) const
+		{
+			return std::hash<decltype(ThreadID.Value)>{}(ThreadID.Value);
+		}
+	};
+
+	const size_t NumThreads = 10;
+	std::vector<std::thread> Threads(NumThreads);
+	std::vector<AutoRTFM::FThreadID> ThreadIDs(NumThreads);
+	FWaitGroup WaitGroup{NumThreads};
+	for (size_t I = 0; I < 10; I++)
+	{
+		Threads[I] = std::thread([I, &WaitGroup, &ThreadIDs]
+		{
+			ThreadIDs[I] = AutoRTFM::FThreadID::GetCurrent();
+			WaitGroup.Done();
+		});
+	}
+
+	// Wait for all the threads to populate ThreadIDs
+	WaitGroup.Wait();
+
+	// Check that all the thread identifiers are unique
+	std::unordered_set<AutoRTFM::FThreadID, FThreadIDHasher> Set;
+	Set.emplace(AutoRTFM::FThreadID::GetCurrent());
+	for (AutoRTFM::FThreadID ThreadID : ThreadIDs)
+	{
+		bool bUnique = Set.emplace(ThreadID).second;
+		REQUIRE(bUnique);
+	}
+	
+	// Join the threads
+	for (std::thread& Thread : Threads)
+	{
+		Thread.join();
+	}
+}
diff --git a/Engine/Source/Runtime/Core/Private/AutoRTFM/Context.cpp b/Engine/Source/Runtime/Core/Private/AutoRTFM/Context.cpp
--- a/Engine/Source/Runtime/Core/Private/AutoRTFM/Context.cpp
+++ b/Engine/Source/Runtime/Core/Private/AutoRTFM/Context.cpp
@@ -1,6 +1,7 @@
 // Copyright Epic Games, Inc. All Rights Reserved.
 
 #if (defined(__AUTORTFM) && __AUTORTFM)
+
 #include "Context.h"
 
 #include "AutoRTFM/AutoRTFMMetrics.h"
@@ -8,7 +9,6 @@
 #include "ContextInlines.h"
 #include "ExternAPI.h"
 #include "FunctionMap.h"
-#include "GlobalData.h"
 #include "Misc/AssertionMacros.h"
 #include "ScopedGuard.h"
 #include "StackRange.h"
@@ -16,10 +16,6 @@
 #include "Transaction.h"
 #include "TransactionInlines.h"
 
-#include "Templates/UniquePtr.h"
-#include "Containers/StringConv.h"
-#include "HAL/PlatformTLS.h"
-
 #if PLATFORM_WINDOWS
 extern "C" __declspec(dllimport) void __stdcall GetCurrentThreadStackLimits(void**, void**);
 #endif
@@ -425,12 +421,12 @@ ETransactionResult FContext::Transact(void (*UninstrumentedFunction)(void*), voi
 
 	ETransactionResult Result = ETransactionResult::Committed; // Initialize to something to make the compiler happy.
 
-    if (!CurrentTransaction)
-    {
-        ASSERT(Status == EContextStatus::Idle);
+	if (!CurrentTransaction)
+	{
+		ASSERT(Status == EContextStatus::Idle);
 
-		ASSERT(FPlatformTLS::InvalidTlsSlot == CurrentThreadId);
-		CurrentThreadId = FPlatformTLS::GetCurrentThreadId();
+		ASSERT(CurrentThreadId == FThreadID::Invalid);
+		CurrentThreadId = FThreadID::GetCurrent();
 
 		ASSERT(Stack == FStackRange{});
 
@@ -591,7 +587,7 @@ ETransactionResult FContext::Transact(void (*UninstrumentedFunction)(void*), voi
 		// This transaction is within another transaction
 		ASSERT(Status == EContextStatus::OnTrack);
 
-		ASSERT(CurrentThreadId == FPlatformTLS::GetCurrentThreadId());
+		ASSERT(CurrentThreadId == FThreadID::GetCurrent());
 
 		ASSERT(Stack.Contains(TransactStackStart));
 
@@ -679,9 +675,9 @@ void FContext::AbortByLanguageAndThrow()
 
 void FContext::Reset()
 {
-	ASSERT(CurrentThreadId == FPlatformTLS::GetCurrentThreadId() || CurrentThreadId == FPlatformTLS::InvalidTlsSlot);
+	ASSERT(CurrentThreadId == FThreadID::GetCurrent() || CurrentThreadId == FThreadID::Invalid);
 
-	CurrentThreadId = FPlatformTLS::InvalidTlsSlot;
+	CurrentThreadId = FThreadID::Invalid;
 	Stack = {};
 	CurrentTransaction = nullptr;
 	CurrentNest = nullptr;
diff --git a/Engine/Source/Runtime/Core/Private/AutoRTFM/Context.h b/Engine/Source/Runtime/Core/Private/AutoRTFM/Context.h
--- a/Engine/Source/Runtime/Core/Private/AutoRTFM/Context.h
+++ b/Engine/Source/Runtime/Core/Private/AutoRTFM/Context.h
@@ -3,12 +3,11 @@
 #pragma once
 
 #include "AutoRTFM/AutoRTFM.h"
-#include "ContextStatus.h"
-#include "HAL/PlatformTLS.h"
 #include "IntrusivePool.h"
 #include "StackRange.h"
 #include "OpenHashThrottler.h"
 #include "TaskArray.h"
+#include "ThreadID.h"
 #include "Transaction.h"
 
 namespace AutoRTFM
@@ -86,7 +85,7 @@ public:
 	}
 
 	inline FCallNest* GetCurrentNest() const { return CurrentNest; }
-	inline EContextStatus GetStatus() const { return CurrentThreadId == FPlatformTLS::GetCurrentThreadId() ? Status : EContextStatus::Idle; }
+	inline EContextStatus GetStatus() const { return CurrentThreadId == FThreadID::GetCurrent() ? Status : EContextStatus::Idle; }
 	[[noreturn]] void Throw();
 
 	// Returns the starting stack address of the innermost call to Closed(), or
@@ -171,7 +170,7 @@ private:
     FStackRange Stack;
 	void* ClosedStackAddress = nullptr;
     EContextStatus Status{EContextStatus::Idle};
-	uint32 CurrentThreadId{ FPlatformTLS::InvalidTlsSlot };
+	FThreadID CurrentThreadId;
 	uint32 StackLocalInitializerDepth = 0;
 	TTask<void()> PostAbortCallback;
 	FTaskPool TaskPool;
diff --git a/Engine/Source/Runtime/Core/Private/AutoRTFM/StdLibSupport.cpp b/Engine/Source/Runtime/Core/Private/AutoRTFM/StdLibSupport.cpp
--- a/Engine/Source/Runtime/Core/Private/AutoRTFM/StdLibSupport.cpp
+++ b/Engine/Source/Runtime/Core/Private/AutoRTFM/StdLibSupport.cpp
@@ -658,6 +658,7 @@ UE_AUTORTFM_REGISTER_OPEN_TO_CLOSED_FUNCTIONS(
 	UE_AUTORTFM_MAP_OPEN_TO_SELF(localtime_r),
 	UE_AUTORTFM_MAP_OPEN_TO_SELF(bcmp),
 	UE_AUTORTFM_MAP_OPEN_TO_SELF(pthread_getspecific),
+	UE_AUTORTFM_MAP_OPEN_TO_SELF(pthread_self),
 	UE_AUTORTFM_MAP_OPEN_TO_CLOSED(strtof32, (RTFM_StringToFloat<float, char, strtof32>)),
 	UE_AUTORTFM_MAP_OPEN_TO_CLOSED(strtof64, (RTFM_StringToFloat<double, char, strtof64>)),
 	UE_AUTORTFM_MAP_OPEN_TO_CLOSED(stat, RTFM_stat),
diff --git a/Engine/Source/Runtime/Core/Private/AutoRTFM/ThreadID.cpp b/Engine/Source/Runtime/Core/Private/AutoRTFM/ThreadID.cpp
new file mode 100644
--- /dev/null
+++ b/Engine/Source/Runtime/Core/Private/AutoRTFM/ThreadID.cpp
@@ -0,0 +1,18 @@
+// Copyright Epic Games, Inc. All Rights Reserved.
+
+#include "ThreadID.h"
+
+#include <thread>
+
+namespace AutoRTFM
+{
+
+const FThreadID FThreadID::Invalid;
+
+UE_AUTORTFM_ALWAYS_OPEN
+FThreadID FThreadID::GetCurrent()
+{
+	return FThreadID{std::this_thread::get_id()};
+}
+
+}
diff --git a/Engine/Source/Runtime/Core/Private/AutoRTFM/ThreadID.h b/Engine/Source/Runtime/Core/Private/AutoRTFM/ThreadID.h
new file mode 100644
--- /dev/null
+++ b/Engine/Source/Runtime/Core/Private/AutoRTFM/ThreadID.h
@@ -0,0 +1,37 @@
+// Copyright Epic Games, Inc. All Rights Reserved.
+
+#pragma once
+
+#include "AutoRTFM/AutoRTFM.h"
+#include <thread>
+
+namespace AutoRTFM
+{
+
+// A unique identifier for a thread of execution.
+// Constructed as FThreadID::Invalid.
+struct FThreadID 
+{
+	// An invalid thread identifier.
+	static const FThreadID Invalid;
+
+	// Returns the currently executing thread's unique identifier.
+	UE_AUTORTFM_API
+	static FThreadID GetCurrent();
+	
+	// Equality operator
+	inline bool operator == (const FThreadID& Other) const
+	{
+		return Value == Other.Value;
+	}
+
+	// Inequality operator
+	inline bool operator != (const FThreadID& Other) const
+	{
+		return Value != Other.Value;
+	}
+
+	std::thread::id Value;
+};
+
+}
diff --git a/Engine/Source/Runtime/Core/Private/AutoRTFM/Transaction.cpp b/Engine/Source/Runtime/Core/Private/AutoRTFM/Transaction.cpp
--- a/Engine/Source/Runtime/Core/Private/AutoRTFM/Transaction.cpp
+++ b/Engine/Source/Runtime/Core/Private/AutoRTFM/Transaction.cpp
@@ -3,6 +3,7 @@
 #if (defined(__AUTORTFM) && __AUTORTFM)
 
 #include "Transaction.h"
+#include "ContextStatus.h"
 #include "CallNestInlines.h"
 #include "GlobalData.h"
 #include "HAL/Platform.h"
