Fix `CSV_EVENT` under AutoRTFM and add some tests.

diff --git a/Engine/Source/Programs/AutoRTFMTests/AutoRTFMTests.Target.cs b/Engine/Source/Programs/AutoRTFMTests/AutoRTFMTests.Target.cs
--- a/Engine/Source/Programs/AutoRTFMTests/AutoRTFMTests.Target.cs
+++ b/Engine/Source/Programs/AutoRTFMTests/AutoRTFMTests.Target.cs
@@ -72,6 +72,10 @@ public class AutoRTFMTestsTarget : TargetRules
 		// Disable AutoRTFM runtime by default when compiler support is enabled
 		GlobalDefinitions.Add("UE_AUTORTFM_ENABLED_RUNTIME_BY_DEFAULT=0");
 
+		// CsvProfiler tests rely on these being enabled!
+		GlobalDefinitions.Add("CSV_PROFILER=1");
+		GlobalDefinitions.Add("CSV_PROFILER_ENABLE_IN_SHIPPING=1");
+
 		// For testing insights, this sets `UE_TRACE_ENABLED`.
 		bEnableTrace = true;
 	}
diff --git a/Engine/Source/Programs/AutoRTFMTests/Private/AutoRTFMTests.cpp b/Engine/Source/Programs/AutoRTFMTests/Private/AutoRTFMTests.cpp
--- a/Engine/Source/Programs/AutoRTFMTests/Private/AutoRTFMTests.cpp
+++ b/Engine/Source/Programs/AutoRTFMTests/Private/AutoRTFMTests.cpp
@@ -75,7 +75,7 @@ int RunTests(int ArgC, const char* ArgV[])
 		}
 	}
 
-	const TCHAR* CommandLine = TEXT("-Multiprocess -LogCmds=\"LogStreaming off, LogUObjectGlobals off, LogPackageName off, LogAutoRTFM warning\" -AsyncLoadingThread");
+	const TCHAR* CommandLine = TEXT("-Multiprocess -LogCmds=\"LogCsvProfiler off, LogStreaming off, LogUObjectGlobals off, LogPackageName off, LogAutoRTFM warning\" -AsyncLoadingThread");
 	GEngineLoop.PreInit(CommandLine); // Note: Initializes AutoRTFM
 
 	if (bRetryNestedToo)
diff --git a/Engine/Source/Programs/AutoRTFMTests/Private/CsvProfilerTests.cpp b/Engine/Source/Programs/AutoRTFMTests/Private/CsvProfilerTests.cpp
new file mode 100644
--- /dev/null
+++ b/Engine/Source/Programs/AutoRTFMTests/Private/CsvProfilerTests.cpp
@@ -0,0 +1,85 @@
+// Copyright Epic Games, Inc. All Rights Reserved.
+
+#include "Async/TaskGraphInterfaces.h"
+#include "API.h"
+#include "Catch2Includes.h"
+#include "CoreTypes.h"
+#include "ProfilingDebugging/CsvProfiler.h"
+#include "Trace/Trace.h"
+#include "Trace/Trace.inl"
+
+#include <AutoRTFM/AutoRTFMTesting.h>
+
+#if !CSV_PROFILER
+#error Test relies on CsvProfiler!
+#endif
+
+CSV_DEFINE_CATEGORY(MyFalseCategory, false);
+CSV_DEFINE_CATEGORY(MyTrueCategory, true);
+
+TEST_CASE("CsvProfiler")
+{
+	struct RAII final
+	{
+		RAII()
+		{
+			FCsvProfiler::Get()->BeginCapture();
+			FCsvProfiler::Get()->BeginFrame();
+			REQUIRE(FCsvProfiler::Get()->IsCapturing());
+		}
+
+		~RAII()
+		{
+			FCsvProfiler::Get()->EndFrame();
+			FCsvProfiler::Get()->EndCapture();
+		}
+	} _;
+
+	AutoRTFM::Testing::Abort([&]
+	{
+		CSV_EVENT(MyFalseCategory, TEXT("HEREWEGO"));
+		AutoRTFM::AbortTransaction();
+	});
+
+	AutoRTFM::Testing::Commit([&]
+	{
+		CSV_EVENT(MyFalseCategory, TEXT("HEREWEGO"));
+	});
+
+	AutoRTFM::Testing::Abort([&]
+	{
+		CSV_EVENT(MyTrueCategory, TEXT("HEREWEGO"));
+		AutoRTFM::AbortTransaction();
+	});
+
+	AutoRTFM::Testing::Commit([&]
+	{
+		CSV_EVENT(MyTrueCategory, TEXT("HEREWEGO"));
+	});
+
+	AutoRTFM::Testing::Abort([&]
+	{
+		CSV_EVENT(MyFalseCategory, TEXT("HEREWEGO"));
+		CSV_SCOPED_TIMING_STAT(MyFalseCategory, Event);
+		AutoRTFM::AbortTransaction();
+	});
+
+	AutoRTFM::Testing::Commit([&]
+	{
+		CSV_EVENT(MyFalseCategory, TEXT("HEREWEGO"));
+		CSV_SCOPED_TIMING_STAT(MyFalseCategory, Event);
+	});
+
+	AutoRTFM::Testing::Abort([&]
+	{
+		CSV_EVENT(MyTrueCategory, TEXT("HEREWEGO"));
+		CSV_SCOPED_TIMING_STAT(MyTrueCategory, Event);
+		AutoRTFM::AbortTransaction();
+	});
+
+	AutoRTFM::Testing::Commit([&]
+	{
+		CSV_EVENT(MyTrueCategory, TEXT("HEREWEGO"));
+		CSV_SCOPED_TIMING_STAT(MyTrueCategory, Event);
+	});
+}
diff --git a/Engine/Source/Runtime/Core/Public/ProfilingDebugging/CsvProfiler.h b/Engine/Source/Runtime/Core/Public/ProfilingDebugging/CsvProfiler.h
--- a/Engine/Source/Runtime/Core/Public/ProfilingDebugging/CsvProfiler.h
+++ b/Engine/Source/Runtime/Core/Public/ProfilingDebugging/CsvProfiler.h
@@ -403,9 +403,12 @@ public:
 	template <typename FmtType, typename... Types>
 	FORCEINLINE static void RecordEventf(int32 CategoryIndex, const FmtType& Fmt, Types... Args)
 	{
-		static_assert(TIsArrayOrRefOfTypeByPredicate<FmtType, TIsCharEncodingCompatibleWithTCHAR>::Value, "Formatting string must be a TCHAR array.");
-		static_assert((TIsValidVariadicFunctionArg<Types>::Value && ...), "Invalid argument(s) passed to FCsvProfiler::RecordEventf");
-		RecordEventfInternal(CategoryIndex, (const TCHAR*)Fmt, Args...);
+		UE_AUTORTFM_OPEN
+		{
+			static_assert(TIsArrayOrRefOfTypeByPredicate<FmtType, TIsCharEncodingCompatibleWithTCHAR>::Value, "Formatting string must be a TCHAR array.");
+			static_assert((TIsValidVariadicFunctionArg<Types>::Value && ...), "Invalid argument(s) passed to FCsvProfiler::RecordEventf");
+			RecordEventfInternal(CategoryIndex, (const TCHAR*)Fmt, Args...);
+		};
 	}
 
 	CORE_API static void BeginSetWaitStat(const char* StatName, const char* FormattedStatName, const char* FormattedStatNameNonCP);
