Fix `FNamePool::Trace` under AutoRTFM.

diff --git a/Engine/Source/Programs/AutoRTFMTests/Private/UECoreTests.cpp b/Engine/Source/Programs/AutoRTFMTests/Private/UECoreTests.cpp
--- a/Engine/Source/Programs/AutoRTFMTests/Private/UECoreTests.cpp
+++ b/Engine/Source/Programs/AutoRTFMTests/Private/UECoreTests.cpp
@@ -286,6 +286,15 @@ TEST_CASE("UECore.FName")
 			REQUIRE(42 == Name.GetNumber());
 		}
 	}
+
+	SECTION("TraceName")
+	{
+		AutoRTFM::Testing::Commit([&]
+		{
+			FName Name(TEXT("WOWWEE"), 42);
+			FName::TraceName(Name);
+		});
+	}
 }
 
 TEST_CASE("UECore.STATIC_FUNCTION_FNAME")
diff --git a/Engine/Source/Runtime/Core/Private/UObject/UnrealNames.cpp b/Engine/Source/Runtime/Core/Private/UObject/UnrealNames.cpp
--- a/Engine/Source/Runtime/Core/Private/UObject/UnrealNames.cpp
+++ b/Engine/Source/Runtime/Core/Private/UObject/UnrealNames.cpp
@@ -2000,6 +2000,7 @@ uint32 FNamePool::NumSlots() const
 }
 
 #if UE_TRACE_ENABLED
+UE_AUTORTFM_ALWAYS_OPEN
 UE::Trace::FEventRef32 FNamePool::Trace(const FNameEntryId& EntryId)
 {
 	const FNameEntryHandle Handle(EntryId);
