[AutoRTFM] AutoRTFMTests build / flakiness fixes

FName::TraceName() is annotated with [[nodiscard]], so needs a (void) before the call.
`FCsvProfiler::Get()->IsCapturing()` appears to be non-deterministic, likely due to the internal threaded task system. Given that these tests aren't actually testing the profiler funcitonality, just remove this check.

diff --git a/Engine/Source/Programs/AutoRTFMTests/Private/CsvProfilerTests.cpp b/Engine/Source/Programs/AutoRTFMTests/Private/CsvProfilerTests.cpp
--- a/Engine/Source/Programs/AutoRTFMTests/Private/CsvProfilerTests.cpp
+++ b/Engine/Source/Programs/AutoRTFMTests/Private/CsvProfilerTests.cpp
@@ -25,7 +25,6 @@ TEST_CASE("CsvProfiler")
 		{
 			FCsvProfiler::Get()->BeginCapture();
 			FCsvProfiler::Get()->BeginFrame();
-			REQUIRE(FCsvProfiler::Get()->IsCapturing());
 		}
 
 		~RAII()
diff --git a/Engine/Source/Programs/AutoRTFMTests/Private/UECoreTests.cpp b/Engine/Source/Programs/AutoRTFMTests/Private/UECoreTests.cpp
--- a/Engine/Source/Programs/AutoRTFMTests/Private/UECoreTests.cpp
+++ b/Engine/Source/Programs/AutoRTFMTests/Private/UECoreTests.cpp
@@ -292,7 +292,7 @@ TEST_CASE("UECore.FName")
 		AutoRTFM::Testing::Commit([&]
 		{
 			FName Name(TEXT("WOWWEE"), 42);
-			FName::TraceName(Name);
+			(void)FName::TraceName(Name);
 		});
 	}
 }
@@ -2901,4 +2901,4 @@ TEST_CASE("UECore.ParallelFor")
 	});
 
 	REQUIRE(Count == Parallelism);
-}
\ No newline at end of file
+}
