[AutoRTFM] More aggressively throttle memory validation hashing

There's a chicken and egg problem with using the measured time taken to hash, to adjust the probability of hashing.
Opens that become active within a throttling profile will use their old probability, which can be from a time when there was far less in the write log. If the write log grows very large, then takes a path that reactivates a number of opens we can take an excessive amount of time hashing the open once. After the initial hash, these will be throttled, but a single hash of each reactivated open is enough to stall for minutes for some content.

The fix here is to uniformly lower the probabilty of all opens (active and inactive) if we're severely over budget.

Other fixes:
* Don't start a new open with a probability of 1. Use the lowest probability across all active opens.
* Correctly measure the number of bytes hashed. This was ignoring the parent's write logs.

diff --git a/Engine/Source/Programs/AutoRTFMTests/Private/OpenHashThrottlerTests.cpp b/Engine/Source/Programs/AutoRTFMTests/Private/OpenHashThrottlerTests.cpp
--- a/Engine/Source/Programs/AutoRTFMTests/Private/OpenHashThrottlerTests.cpp
+++ b/Engine/Source/Programs/AutoRTFMTests/Private/OpenHashThrottlerTests.cpp
@@ -72,20 +72,53 @@ TEST_CASE("OpenHashThrottler")
 
 	SECTION("Exceeds TargetFractionHashing limits")
 	{
-		for (int I = 0; I < 100; I++)
+		SECTION("Small Margin") // limit exceeded by typical margin (1x - 2x)
 		{
-			// 20% of the time spent hashing (greater than the 10% threshold to throttle)
-			Hash(/* DurationNotHashing */ 0.08, /* DurationHashing */ 0.02, OpenReturnAddressA);
-			Hash(/* DurationNotHashing */ 0.08, /* DurationHashing */ 0.02, OpenReturnAddressB);
-			Hash(/* DurationNotHashing */ 0.08, /* DurationHashing */ 0.02, OpenReturnAddressC);
-			Update();
+			for (int I = 0; I < 100; I++)
+			{
+				// 15% of the time spent hashing (greater than the 10% threshold to throttle)
+				Hash(/* DurationNotHashing */ 0.85, /* DurationHashing */ 0.15, OpenReturnAddressA);
+				Hash(/* DurationNotHashing */ 0.85, /* DurationHashing */ 0.15, OpenReturnAddressB);
+				Hash(/* DurationNotHashing */ 0.85, /* DurationHashing */ 0.15, OpenReturnAddressC);
+				Update();
+			}
+			REQUIRE(Throttler.HashProbabilityFor(OpenReturnAddressA) < 1.0);
+			REQUIRE(Throttler.HashProbabilityFor(OpenReturnAddressA) > 0.0);
+			REQUIRE(Throttler.HashProbabilityFor(OpenReturnAddressB) < 1.0);
+			REQUIRE(Throttler.HashProbabilityFor(OpenReturnAddressB) > 0.0);
+			REQUIRE(Throttler.HashProbabilityFor(OpenReturnAddressC) < 1.0);
+			REQUIRE(Throttler.HashProbabilityFor(OpenReturnAddressC) > 0.0);
+		}
+		
+		SECTION("Excessive") // limit exceeded by large margin (2x+)
+		{
+			const double OriginalProbabilityC = Throttler.HashProbabilityFor(OpenReturnAddressC);
+			for (int I = 0; I < 100; I++)
+			{
+				// 1% of the time hashing (below throttling level)
+				Hash(/* DurationNotHashing */ 0.99, /* DurationHashing */ 0.01, OpenReturnAddressC);
+				Update();
+			}
+			// Below throttling threshold, should not have been updated.
+			REQUIRE(OriginalProbabilityC == Throttler.HashProbabilityFor(OpenReturnAddressC));
+			
+			for (int I = 0; I < 100; I++)
+			{
+				// 30% of the time spent hashing (greater than the 10% threshold to throttle and
+				// more then 2x the budget)
+				Hash(/* DurationNotHashing */ 0.60, /* DurationHashing */ 0.40, OpenReturnAddressA);
+				Hash(/* DurationNotHashing */ 0.70, /* DurationHashing */ 0.30, OpenReturnAddressB);
+				Update();
+			}
+
+			// When excessively exceeding the budget, all open probabilities are
+			// uniformly reduced, including inactive opens.
+			REQUIRE(Throttler.HashProbabilityFor(OpenReturnAddressA) < 1.0);
+			REQUIRE(Throttler.HashProbabilityFor(OpenReturnAddressA) > 0.0);
+			REQUIRE(Throttler.HashProbabilityFor(OpenReturnAddressB) == Throttler.HashProbabilityFor(OpenReturnAddressA));
+			REQUIRE(Throttler.HashProbabilityFor(OpenReturnAddressC) < OriginalProbabilityC);
+			REQUIRE(Throttler.HashProbabilityFor(OpenReturnAddressC) > 0.0);
 		}
-		REQUIRE(Throttler.HashProbabilityFor(OpenReturnAddressA) < 1.0);
-		REQUIRE(Throttler.HashProbabilityFor(OpenReturnAddressA) > 0.0);
-		REQUIRE(Throttler.HashProbabilityFor(OpenReturnAddressB) < 1.0);
-		REQUIRE(Throttler.HashProbabilityFor(OpenReturnAddressB) > 0.0);
-		REQUIRE(Throttler.HashProbabilityFor(OpenReturnAddressC) < 1.0);
-		REQUIRE(Throttler.HashProbabilityFor(OpenReturnAddressC) > 0.0);
 	}
 
 	SECTION("Normalize by duration")
@@ -94,17 +127,17 @@ TEST_CASE("OpenHashThrottler")
 		{
 			// Time spent hashing ratios:
 			// OpenReturnAddressA:OpenReturnAddressB:OpenReturnAddressC - 5:3:1 
-			Hash(/* DurationNotHashing */ 0.05, /* DurationHashing */ 0.05, OpenReturnAddressA);
-			Hash(/* DurationNotHashing */ 0.05, /* DurationHashing */ 0.03, OpenReturnAddressB);
-			Hash(/* DurationNotHashing */ 0.05, /* DurationHashing */ 0.01, OpenReturnAddressC);
+			Hash(/* DurationNotHashing */ 0.25, /* DurationHashing */ 0.05, OpenReturnAddressA);
+			Hash(/* DurationNotHashing */ 0.25, /* DurationHashing */ 0.03, OpenReturnAddressB);
+			Hash(/* DurationNotHashing */ 0.25, /* DurationHashing */ 0.01, OpenReturnAddressC);
 			Update();
 		}
 		REQUIRE(Throttler.HashProbabilityFor(OpenReturnAddressA) < 1.0);
 		REQUIRE(Throttler.HashProbabilityFor(OpenReturnAddressA) > 0.0);
 		REQUIRE(Throttler.HashProbabilityFor(OpenReturnAddressB) < 1.0);
 		REQUIRE(Throttler.HashProbabilityFor(OpenReturnAddressB) > 0.0);
-		REQUIRE(Throttler.HashProbabilityFor(OpenReturnAddressC) < 1.0);
-		REQUIRE(Throttler.HashProbabilityFor(OpenReturnAddressC) > 0.0);
+		REQUIRE(Throttler.HashProbabilityFor(OpenReturnAddressC) <= 1.0);
+		REQUIRE(Throttler.HashProbabilityFor(OpenReturnAddressC) >  0.0);
 		REQUIRE(Throttler.HashProbabilityFor(OpenReturnAddressA) < Throttler.HashProbabilityFor(OpenReturnAddressB));
 		REQUIRE(Throttler.HashProbabilityFor(OpenReturnAddressB) < Throttler.HashProbabilityFor(OpenReturnAddressC));
 	}
@@ -115,27 +148,42 @@ TEST_CASE("OpenHashThrottler")
 		{
 			// Hash count ratios:
 			// OpenReturnAddressA:OpenReturnAddressB:OpenReturnAddressC - 5:3:1 
-			Hash(/* DurationNotHashing */ 0.05, /* DurationHashing */ 0.05, OpenReturnAddressA);
-			Hash(/* DurationNotHashing */ 0.05, /* DurationHashing */ 0.05, OpenReturnAddressB);
-			Hash(/* DurationNotHashing */ 0.05, /* DurationHashing */ 0.05, OpenReturnAddressA);
-			Hash(/* DurationNotHashing */ 0.05, /* DurationHashing */ 0.05, OpenReturnAddressC);
-			Hash(/* DurationNotHashing */ 0.05, /* DurationHashing */ 0.05, OpenReturnAddressA);
-			Hash(/* DurationNotHashing */ 0.05, /* DurationHashing */ 0.05, OpenReturnAddressB);
-			Hash(/* DurationNotHashing */ 0.05, /* DurationHashing */ 0.05, OpenReturnAddressA);
-			Hash(/* DurationNotHashing */ 0.05, /* DurationHashing */ 0.05, OpenReturnAddressB);
-			Hash(/* DurationNotHashing */ 0.05, /* DurationHashing */ 0.05, OpenReturnAddressA);
+			Hash(/* DurationNotHashing */ 0.25, /* DurationHashing */ 0.05, OpenReturnAddressA);
+			Hash(/* DurationNotHashing */ 0.25, /* DurationHashing */ 0.05, OpenReturnAddressB);
+			Hash(/* DurationNotHashing */ 0.25, /* DurationHashing */ 0.05, OpenReturnAddressA);
+			Hash(/* DurationNotHashing */ 0.25, /* DurationHashing */ 0.05, OpenReturnAddressC);
+			Hash(/* DurationNotHashing */ 0.25, /* DurationHashing */ 0.05, OpenReturnAddressA);
+			Hash(/* DurationNotHashing */ 0.25, /* DurationHashing */ 0.05, OpenReturnAddressB);
+			Hash(/* DurationNotHashing */ 0.25, /* DurationHashing */ 0.05, OpenReturnAddressA);
+			Hash(/* DurationNotHashing */ 0.25, /* DurationHashing */ 0.05, OpenReturnAddressB);
+			Hash(/* DurationNotHashing */ 0.25, /* DurationHashing */ 0.05, OpenReturnAddressA);
 			Update();
 		}
 		REQUIRE(Throttler.HashProbabilityFor(OpenReturnAddressA) < 1.0);
 		REQUIRE(Throttler.HashProbabilityFor(OpenReturnAddressA) > 0.0);
 		REQUIRE(Throttler.HashProbabilityFor(OpenReturnAddressB) < 1.0);
 		REQUIRE(Throttler.HashProbabilityFor(OpenReturnAddressB) > 0.0);
-		REQUIRE(Throttler.HashProbabilityFor(OpenReturnAddressC) < 1.0);
-		REQUIRE(Throttler.HashProbabilityFor(OpenReturnAddressC) > 0.0);
+		REQUIRE(Throttler.HashProbabilityFor(OpenReturnAddressC) <= 1.0);
+		REQUIRE(Throttler.HashProbabilityFor(OpenReturnAddressC) >  0.0);
 		REQUIRE(Throttler.HashProbabilityFor(OpenReturnAddressA) < Throttler.HashProbabilityFor(OpenReturnAddressB));
 		REQUIRE(Throttler.HashProbabilityFor(OpenReturnAddressB) < Throttler.HashProbabilityFor(OpenReturnAddressC));
 	}
 
+	SECTION("New opens inherit lowest probability")
+	{
+		for (int I = 0; I < 100; I++)
+		{
+			Hash(/* DurationNotHashing */ 0.25, /* DurationHashing */ 0.05, OpenReturnAddressA);
+			Hash(/* DurationNotHashing */ 0.25, /* DurationHashing */ 0.03, OpenReturnAddressB);
+			Update();
+		}
+
+		Hash(/* DurationNotHashing */ 0.20, /* DurationHashing */ 0.01, OpenReturnAddressC);
+		
+		REQUIRE(Throttler.HashProbabilityFor(OpenReturnAddressA) < Throttler.HashProbabilityFor(OpenReturnAddressB));
+		REQUIRE(Throttler.HashProbabilityFor(OpenReturnAddressC) == Throttler.HashProbabilityFor(OpenReturnAddressA));
+	}
+
 	SECTION("Preserve probabilities for old open return addresses")
 	{
 		// Throttle OpenReturnAddressA
@@ -167,9 +215,9 @@ TEST_CASE("OpenHashThrottler")
 	{
 		for (int I = 0; I < 100; I++)
 		{
-			Hash(/* DurationNotHashing */ 0.05, /* DurationHashing */ 0.2, OpenReturnAddressA);
-			Hash(/* DurationNotHashing */ 0.05, /* DurationHashing */ 0.2, OpenReturnAddressB);
-			Hash(/* DurationNotHashing */ 0.05, /* DurationHashing */ 0.2, OpenReturnAddressC);
+			Hash(/* DurationNotHashing */ 0.5, /* DurationHashing */ 0.2, OpenReturnAddressA);
+			Hash(/* DurationNotHashing */ 0.5, /* DurationHashing */ 0.2, OpenReturnAddressB);
+			Hash(/* DurationNotHashing */ 0.5, /* DurationHashing */ 0.2, OpenReturnAddressC);
 			Update(1.0);
 		}
 
@@ -179,9 +227,9 @@ TEST_CASE("OpenHashThrottler")
 
 		for (int I = 0; I < 100; I++)
 		{
-			Hash(/* DurationNotHashing */ 0.05, /* DurationHashing */ 0.1, OpenReturnAddressA);
-			Hash(/* DurationNotHashing */ 0.05, /* DurationHashing */ 0.1, OpenReturnAddressB);
-			Hash(/* DurationNotHashing */ 0.05, /* DurationHashing */ 0.1, OpenReturnAddressC);
+			Hash(/* DurationNotHashing */ 0.5, /* DurationHashing */ 0.1, OpenReturnAddressA);
+			Hash(/* DurationNotHashing */ 0.5, /* DurationHashing */ 0.1, OpenReturnAddressB);
+			Hash(/* DurationNotHashing */ 0.5, /* DurationHashing */ 0.1, OpenReturnAddressC);
 			Update(1.0);
 		}
 
@@ -194,12 +242,12 @@ TEST_CASE("OpenHashThrottler")
 	{
 		for (int I = 0; I < 100; I++)
 		{
-			Hash(/* DurationNotHashing */ 0.01, /* DurationHashing */ 0.1, OpenReturnAddressA);
-			Hash(/* DurationNotHashing */ 0.01, /* DurationHashing */ 0.1, OpenReturnAddressB);
-			Hash(/* DurationNotHashing */ 0.01, /* DurationHashing */ 0.1, OpenReturnAddressC);
-			Hash(/* DurationNotHashing */ 0.01, /* DurationHashing */ 0.1, OpenReturnAddressA);
-			Hash(/* DurationNotHashing */ 0.01, /* DurationHashing */ 0.1, OpenReturnAddressB);
-			Hash(/* DurationNotHashing */ 0.01, /* DurationHashing */ 0.1, OpenReturnAddressC);
+			Hash(/* DurationNotHashing */ 0.25, /* DurationHashing */ 0.10, OpenReturnAddressA);
+			Hash(/* DurationNotHashing */ 0.25, /* DurationHashing */ 0.10, OpenReturnAddressB);
+			Hash(/* DurationNotHashing */ 0.25, /* DurationHashing */ 0.10, OpenReturnAddressC);
+			Hash(/* DurationNotHashing */ 0.25, /* DurationHashing */ 0.10, OpenReturnAddressA);
+			Hash(/* DurationNotHashing */ 0.25, /* DurationHashing */ 0.10, OpenReturnAddressB);
+			Hash(/* DurationNotHashing */ 0.25, /* DurationHashing */ 0.10, OpenReturnAddressC);
 			Update(1.0);
 		}
 
@@ -209,9 +257,9 @@ TEST_CASE("OpenHashThrottler")
 
 		for (int I = 0; I < 100; I++)
 		{
-			Hash(/* DurationNotHashing */ 0.02, /* DurationHashing */ 0.1, OpenReturnAddressA);
-			Hash(/* DurationNotHashing */ 0.02, /* DurationHashing */ 0.1, OpenReturnAddressB);
-			Hash(/* DurationNotHashing */ 0.02, /* DurationHashing */ 0.1, OpenReturnAddressC);
+			Hash(/* DurationNotHashing */ 0.35, /* DurationHashing */ 0.10, OpenReturnAddressA);
+			Hash(/* DurationNotHashing */ 0.35, /* DurationHashing */ 0.10, OpenReturnAddressB);
+			Hash(/* DurationNotHashing */ 0.35, /* DurationHashing */ 0.10, OpenReturnAddressC);
 			Update(1.0);
 		}
 
diff --git a/Engine/Source/Runtime/Core/Private/AutoRTFM/OpenHashThrottler.cpp b/Engine/Source/Runtime/Core/Private/AutoRTFM/OpenHashThrottler.cpp
--- a/Engine/Source/Runtime/Core/Private/AutoRTFM/OpenHashThrottler.cpp
+++ b/Engine/Source/Runtime/Core/Private/AutoRTFM/OpenHashThrottler.cpp
@@ -37,11 +37,16 @@ FOpenHashThrottler::FOpenHashThrottler(FSeconds LogInterval, FSeconds AdjustThro
 void FOpenHashThrottler::OnHash(FSeconds StartTime, FSeconds EndTime, const void* OpenReturnAddress, const FWriteLog& WriteLog)
 {
 	const FSeconds HashDuration = EndTime - StartTime;
+	ThrottlingData.TotalTimeSpentHashing += HashDuration;
+	
 	FOpenAddressInfo& OpenInfo = ThrottlingData.Opens.FindOrAdd(OpenReturnAddress);
 	OpenInfo.TimeSpentHashing += HashDuration;
 	OpenInfo.bActive = true;
-	ThrottlingData.TotalTimeSpentHashing += HashDuration;
-
+	if (OpenInfo.Probability == 0)
+	{
+		OpenInfo.Probability = ThrottlingData.DefaultHashProbability;
+	}
+	
 	LogStats.TimeSpentHashing += HashDuration;
 	LogStats.NumHashCalls++;
 	LogStats.NumBytesHashed += WriteLog.TotalSize();
@@ -50,14 +55,13 @@ void FOpenHashThrottler::OnHash(FSeconds StartTime, FSeconds EndTime, const void
 
 double FOpenHashThrottler::HashProbabilityFor(const void* OpenReturnAddress)
 {
-	if (FOpenAddressInfo* Info = ThrottlingData.Opens.Find(OpenReturnAddress))
+	FOpenAddressInfo& OpenInfo = ThrottlingData.Opens.FindOrAdd(OpenReturnAddress);
+	if (OpenInfo.Probability == 0)
 	{
-		Info->bActive = true;
-		return Info->Probability;
+		OpenInfo.Probability = ThrottlingData.DefaultHashProbability;
 	}
-
-	// No data for this open return address - default to 1.
-	return 1.0;
+	OpenInfo.bActive = true;
+	return OpenInfo.Probability;
 }
 
 bool FOpenHashThrottler::ShouldHashFor(const void* OpenReturnAddress)
@@ -128,63 +132,94 @@ void FOpenHashThrottler::UpdateThrottlingData()
 
 	// Probability multiplier to reach the target time spent hashing.
 	const double TotalGain = TargetFractionHashing * ThrottlingData.TimeSinceLastUpdate / ThrottlingData.TotalTimeSpentHashing;
+	
+	// Something non-zero, so we can scale probabilities back up with multiplications.
+	const double MinProbability = 1e-6;
 
-	// Average time spent hashing per open
-	const double AverageTimeSpentHashingPerOpen =
-		ThrottlingData.TotalTimeSpentHashing / static_cast<double>(FMath::Max<size_t>(NumOpensThatHashed, 1));
-
-	// Apply the probability multiplier and normalize the time spent in each open.
-	for (auto& It : ThrottlingData.Opens)
+	if (TotalGain < 0.5)
 	{
-		if (!It.Value.bActive)
-		{
-			continue;
-		}
-
-		// Calculate the target probability to normalize the time spent in each open, and
-		// to aim for the target total fractional time spent in hashing (TargetFractionHashing).
-
-		// Something non-zero, so we can scale probabilities back up with multiplications.
-		const double MinProbability = 1e-6;
-
-		// The time spent hashing this open this update
-		const FSeconds TimeSpentHashing = It.Value.TimeSpentHashing;
-		
-		// Target probability starts with the current probability
-		double TargetProbability = FMath::Max(It.Value.Probability, MinProbability);
-
-		// Adjust for the relative time spent hashing this open compared to the others.
-		if (TimeSpentHashing > 0)
+		// We've exceeded our budget by 2x or greater.
+		// Instead of adjusting each of the open hash probabilities individually
+		// to normalize the probabilities based on time spent hashing, apply the
+		// total gain to all opens. This is done to prevent long hashing stalls
+		// when there are sudden increases of the write log size. In this
+		// situation, newly active active can have a probability that is
+		// substantially higher than those that have been tuned, and allowing
+		// each of these opens to hash even once in an update can dramatically
+		// exceed the budgeted time.
+		for (auto& It : ThrottlingData.Opens)
 		{
-			TargetProbability *= AverageTimeSpentHashingPerOpen / TimeSpentHashing;
+			It.Value.Probability = FMath::Max(It.Value.Probability * TotalGain, MinProbability);
+			It.Value.TimeSpentHashing = 0;
+			It.Value.bActive = false;
 		}
+		ThrottlingData.DefaultHashProbability = FMath::Max(ThrottlingData.DefaultHashProbability * TotalGain, MinProbability);
+	}
+	else
+	{
+		// Average time spent hashing per open
+		const double AverageTimeSpentHashingPerOpen =
+			ThrottlingData.TotalTimeSpentHashing / static_cast<double>(FMath::Max<size_t>(NumOpensThatHashed, 1));
 
-		// If the open was hashed this update, or probabilities are being raised
-		// then adjust by the total gain for all opens.
-		if (TimeSpentHashing > 0 || TotalGain > 1)
-		{
-			TargetProbability *= TotalGain;
-		}
+		// The new lowest probability across all active opens.
+		double LowestProbability = 1;
 
-		// Finally clamp between 0..1
-		TargetProbability = FMath::Clamp(TargetProbability, 0, 1);
-		
-		if (TargetProbability < It.Value.Probability)
-		{
-			// Probability is being reduced.
-			// Apply target probability immediately to ensure the application doesn't stall.
-			It.Value.Probability = TargetProbability;
-		}
-		else
+		// Apply the probability multiplier and normalize the time spent in each open.
+		for (auto& It : ThrottlingData.Opens)
 		{
-			// Probability is being increased.
-			// Interpolate to the new target probability at 15% per second.
-			It.Value.Probability = FMath::Lerp(TargetProbability, It.Value.Probability, FMath::Pow(0.85, ThrottlingData.TimeSinceLastUpdate));
+			if (!It.Value.bActive)
+			{
+				continue;
+			}
+
+			// Calculate the target probability to normalize the time spent in each open, and
+			// to aim for the target total fractional time spent in hashing (TargetFractionHashing).
+
+			// The time spent hashing this open this update
+			const FSeconds TimeSpentHashing = It.Value.TimeSpentHashing;
+			
+			// Target probability starts with the current probability
+			double TargetProbability = FMath::Max(It.Value.Probability, MinProbability);
+
+			// Adjust for the relative time spent hashing this open compared to the others.
+			if (TimeSpentHashing > 0)
+			{
+				TargetProbability *= AverageTimeSpentHashingPerOpen / TimeSpentHashing;
+			}
+
+			// If the open was hashed this update, or probabilities are being raised
+			// then adjust by the total gain for all opens.
+			if (TimeSpentHashing > 0 || TotalGain > 1)
+			{
+				TargetProbability *= TotalGain;
+			}
+
+			// Finally clamp between 0..1
+			TargetProbability = FMath::Clamp(TargetProbability, 0, 1);
+			
+			if (TargetProbability < It.Value.Probability)
+			{
+				// Probability is being reduced.
+				// Apply target probability immediately to ensure the application doesn't stall.
+				It.Value.Probability = TargetProbability;
+			}
+			else
+			{
+				// Probability is being increased.
+				// Interpolate to the new target probability at 15% per second.
+				It.Value.Probability = FMath::Lerp(TargetProbability, It.Value.Probability, FMath::Pow(0.85, ThrottlingData.TimeSinceLastUpdate));
+			}
+
+			// Update LowestProbability if this is the new lowest probability.
+			LowestProbability = FMath::Min(LowestProbability, It.Value.Probability);
+
+			// Reset the TimeSpentHashing and bActive state for this open.
+			It.Value.TimeSpentHashing = 0;
+			It.Value.bActive = false;
 		}
 
-		// Reset the TimeSpentHashing and bActive state for this open.
-		It.Value.TimeSpentHashing = 0;
-		It.Value.bActive = false;
+		// Default new opens with the lowest probability of all active opens.
+		ThrottlingData.DefaultHashProbability = LowestProbability;
 	}
 
 	ThrottlingData.TotalTimeSpentHashing = 0;
diff --git a/Engine/Source/Runtime/Core/Private/AutoRTFM/OpenHashThrottler.h b/Engine/Source/Runtime/Core/Private/AutoRTFM/OpenHashThrottler.h
--- a/Engine/Source/Runtime/Core/Private/AutoRTFM/OpenHashThrottler.h
+++ b/Engine/Source/Runtime/Core/Private/AutoRTFM/OpenHashThrottler.h
@@ -78,7 +78,7 @@ private:
 		// Total time spent hashing this open since last reset.
 		FSeconds TimeSpentHashing = 0;
 		// Throttled probability to hash this open.
-		double Probability = 1;
+		double Probability = 0;
 		// True if the open address was queried or hashed since last throttling update.
 		bool bActive = true;
 	};
@@ -92,6 +92,9 @@ private:
 		FSeconds TotalTimeSpentHashing = 0;
 		// A map of open return address to open info.
 		THashMap<const void*, FOpenAddressInfo> Opens;
+		// The default hash probability if the open return address is not found
+		// in Opens.
+		double DefaultHashProbability = 1;
 	};
 	
 	// Statistics for logging
diff --git a/Engine/Source/Runtime/Core/Private/AutoRTFM/Transaction.cpp b/Engine/Source/Runtime/Core/Private/AutoRTFM/Transaction.cpp
--- a/Engine/Source/Runtime/Core/Private/AutoRTFM/Transaction.cpp
+++ b/Engine/Source/Runtime/Core/Private/AutoRTFM/Transaction.cpp
@@ -513,19 +513,20 @@ void FTransaction::ValidateWriteHash() const
 
 FTransaction::FWriteHash FTransaction::CalculateNestedWriteHash() const
 {
-	FOpenHashThrottler::FHashScope Profile(Context->GetOpenHashThrottler(), CurrentOpenReturnAddress, WriteLog);
-	return CalculateNestedWriteHashWithLimit(NumWriteLogsHashed);
+	return CalculateNestedWriteHashWithLimit(NumWriteLogsHashed, CurrentOpenReturnAddress);
 }
 
-FTransaction::FWriteHash FTransaction::CalculateNestedWriteHashWithLimit(size_t NumWriteEntries) const
+FTransaction::FWriteHash FTransaction::CalculateNestedWriteHashWithLimit(size_t NumWriteEntries, const void *OpenReturnAddress) const
 {
 	FXxHash64Builder Hasher;
 	if (nullptr != Parent)
 	{
-		FWriteHash ParentHash = Parent->CalculateNestedWriteHashWithLimit(Parent->WriteLog.Num());
+		FWriteHash ParentHash = Parent->CalculateNestedWriteHashWithLimit(Parent->WriteLog.Num(), OpenReturnAddress);
 		Hasher.Update(&ParentHash, sizeof(ParentHash));
 	}
 
+	FOpenHashThrottler::FHashScope Profile(Context->GetOpenHashThrottler(), OpenReturnAddress, WriteLog);
+
 	size_t WriteIndex = 0;
 
 	using FSmallValue = uint64;
diff --git a/Engine/Source/Runtime/Core/Private/AutoRTFM/Transaction.h b/Engine/Source/Runtime/Core/Private/AutoRTFM/Transaction.h
--- a/Engine/Source/Runtime/Core/Private/AutoRTFM/Transaction.h
+++ b/Engine/Source/Runtime/Core/Private/AutoRTFM/Transaction.h
@@ -178,7 +178,8 @@ private:
 	// Returns a hash of all the memory written during this transaction and all parent transactions.
 	// Can be used to ensure that no transaction written memory is modified during an open.
 	// NumWriteEntries is the number of write entries to hash before stopping.
-	FWriteHash CalculateNestedWriteHashWithLimit(size_t NumWriteEntries) const;
+	// OpenReturnAddress is the return address for the open that triggered this hash.
+	FWriteHash CalculateNestedWriteHashWithLimit(size_t NumWriteEntries, const void *OpenReturnAddress) const;
 
 	void CollectStats() const;
 
