Make `FModuleManager::UnloadModule` work in transactions.

diff --git a/Engine/Source/Programs/AutoRTFMTests/Private/UECoreTests.cpp b/Engine/Source/Programs/AutoRTFMTests/Private/UECoreTests.cpp
--- a/Engine/Source/Programs/AutoRTFMTests/Private/UECoreTests.cpp
+++ b/Engine/Source/Programs/AutoRTFMTests/Private/UECoreTests.cpp
@@ -15,6 +15,7 @@
 #include "Internationalization/TextHistory.h"
 #include "Logging/LogMacros.h"
 #include "Logging/StructuredLog.h"
+#include "Modules/ModuleManager.h"
 #include "Serialization/CustomVersion.h"
 #include "UObject/CoreRedirects.h"
 #include "UObject/LinkerInstancingContext.h"
@@ -2902,3 +2903,50 @@ TEST_CASE("UECore.ParallelFor")
 
 	REQUIRE(Count == Parallelism);
 }
+
+TEST_CASE("UECore.ModuleManager")
+{
+	struct RAII final
+	{
+		~RAII()
+		{
+			// We've unloaded the module so of course it isn't loaded!
+			REQUIRE(!FModuleManager::Get().IsModuleLoaded(TEXT("CoreUObject")));
+
+			REQUIRE(FModuleManager::Get().LoadModule(TEXT("CoreUObject")));
+			REQUIRE(FModuleManager::Get().IsModuleLoaded(TEXT("CoreUObject")));
+		}
+	};
+
+	SECTION("POD FName")
+	{
+		RAII _;
+		AutoRTFM::Testing::Abort([&]
+		{
+			REQUIRE(FModuleManager::Get().UnloadModule(TEXT("CoreUObject")));
+			AutoRTFM::AbortTransaction();
+		});
+
+		AutoRTFM::Testing::Commit([&]
+		{
+			REQUIRE(FModuleManager::Get().UnloadModule(TEXT("CoreUObject")));
+		});
+	}
+
+	SECTION("FName")
+	{
+		RAII _;
+		AutoRTFM::Testing::Abort([&]
+		{
+			FName Name(TEXT("CoreUObject"));
+			REQUIRE(FModuleManager::Get().UnloadModule(Name));
+			AutoRTFM::AbortTransaction();
+		});
+
+		AutoRTFM::Testing::Commit([&]
+		{
+			FName Name(TEXT("CoreUObject"));
+			REQUIRE(FModuleManager::Get().UnloadModule(Name));
+		});
+	}
+}
diff --git a/Engine/Source/Runtime/Core/Private/Modules/ModuleManager.cpp b/Engine/Source/Runtime/Core/Private/Modules/ModuleManager.cpp
--- a/Engine/Source/Runtime/Core/Private/Modules/ModuleManager.cpp
+++ b/Engine/Source/Runtime/Core/Private/Modules/ModuleManager.cpp
@@ -1090,6 +1090,17 @@ bool FModuleManager::UnloadModule( const FName InModuleName, bool bIsShutdown, b
 		// Only if already loaded
 		if( ModuleInfo.Module.IsValid() )
 		{
+			// If we are running in a transaction we defer the unload until we know the body of
+			// the transaction is definitely going to be committed.
+			if (AutoRTFM::IsClosed())
+			{
+				AutoRTFM::OnCommit([this, InModuleName, bIsShutdown, bAllowUnloadCode]
+				{
+					this->UnloadModule(InModuleName, bIsShutdown, bAllowUnloadCode);
+				});
+				return true;
+			}
+
 			// Will offer use-before-ready protection at next reload
 			ModuleInfo.bIsReady = false;
 
