Null compare to avoid a catch2 + ASan bug.

diff --git a/Engine/Source/Programs/AutoRTFMTests/Private/UECoreTests.cpp b/Engine/Source/Programs/AutoRTFMTests/Private/UECoreTests.cpp
--- a/Engine/Source/Programs/AutoRTFMTests/Private/UECoreTests.cpp
+++ b/Engine/Source/Programs/AutoRTFMTests/Private/UECoreTests.cpp
@@ -2913,7 +2913,7 @@ TEST_CASE("UECore.ModuleManager")
 			// We've unloaded the module so of course it isn't loaded!
 			REQUIRE(!FModuleManager::Get().IsModuleLoaded(TEXT("CoreUObject")));
 
-			REQUIRE(FModuleManager::Get().LoadModule(TEXT("CoreUObject")));
+			REQUIRE(nullptr != FModuleManager::Get().LoadModule(TEXT("CoreUObject")));
 			REQUIRE(FModuleManager::Get().IsModuleLoaded(TEXT("CoreUObject")));
 		}
 	};
