[AutoRTFM] Fix static-analysis warning about leaking memory.

These unit tests malloc memory and then abort the transaction. It's unclear whether this would be a leak or not (it depends on whether we hook malloc in the unit test harness). I've updated the tests to use `FMemory::Malloc` and `FMemory::Free` which unambiguously clean up transactionally-allocated memory during an abort.

An example of the static-analysis warning is here: https://epicgames.slack.com/archives/C059TRJ6XJ4/p1739490902436389

diff --git a/Engine/Source/Programs/AutoRTFMTests/Private/SharedPointerTests.cpp b/Engine/Source/Programs/AutoRTFMTests/Private/SharedPointerTests.cpp
--- a/Engine/Source/Programs/AutoRTFMTests/Private/SharedPointerTests.cpp
+++ b/Engine/Source/Programs/AutoRTFMTests/Private/SharedPointerTests.cpp
@@ -3,6 +3,7 @@
 #include "Catch2Includes.h"
 #include <AutoRTFM/AutoRTFM.h>
 #include "Templates/SharedPointer.h"
+#include "HAL/UnrealMemory.h"
 
 TEST_CASE("SharedPointer.PreviouslyAllocated")
 {
@@ -144,7 +145,7 @@ TEST_CASE("SharedPointer.AbortNestedTransactionWithNewlyAllocated")
 
 template<typename T> FORCENOINLINE void* MakeMemoryForT()
 {
-	return malloc(sizeof(T));
+	return FMemory::Malloc(sizeof(T));
 }
 
 TEST_CASE("SharedPointer.NestedTransactionWithPlacementNewlyAllocated")
@@ -161,7 +162,7 @@ TEST_CASE("SharedPointer.NestedTransactionWithPlacementNewlyAllocated")
 					Ptr Copy = *Foo;
 					Result = *Copy;
 					reinterpret_cast<Ptr*>(Foo)->~Ptr();
-					free(Memory);
+					FMemory::Free(Memory);
 				});
 		});
 
