[AutoRTFM] Fix memory validation hashing for unaligned writes smaller than 8 bytes.

Move write log hashing from FTransaction to FWriteLog.
Add tests and benchmarks.
Fix the stupid bug it uncovered.

diff --git a/Engine/Source/Programs/AutoRTFMTests/Private/Catch2Includes.h b/Engine/Source/Programs/AutoRTFMTests/Private/Catch2Includes.h
--- a/Engine/Source/Programs/AutoRTFMTests/Private/Catch2Includes.h
+++ b/Engine/Source/Programs/AutoRTFMTests/Private/Catch2Includes.h
@@ -5,7 +5,7 @@
 // we need to include the MainCPP include *before* we include catch_amalgamated or we get issues with windows symbols
 #include "CoreMinimal.h"
 #include "LaunchEngineLoop.h"
-#include <AutoRTFM/AutoRTFM.h>
+#include "AutoRTFM/AutoRTFM.h"
 
 #if defined(REQUIRE) || defined(CATCH_REQUIRE)
 #error "Catch2 already included"
@@ -74,6 +74,7 @@ THIRD_PARTY_INCLUDES_END
 #define TEST_CASE(...) CATCH_TEST_CASE(__VA_ARGS__)
 #define TEMPLATE_TEST_CASE(...) CATCH_TEMPLATE_TEST_CASE(__VA_ARGS__)
 #define BENCHMARK(...) CATCH_BENCHMARK(__VA_ARGS__)
+#define BENCHMARK_ADVANCED(...) CATCH_BENCHMARK_ADVANCED(__VA_ARGS__)
 #define SECTION(...) AUTORTFM_CATCH2_SECTION(INTERNAL_CATCH_UNIQUE_NAME(catch_internal_Section), __VA_ARGS__)
 
 namespace AutoRTFMCatch2
diff --git a/Engine/Source/Programs/AutoRTFMTests/Private/UtilsTests.cpp b/Engine/Source/Programs/AutoRTFMTests/Private/UtilsTests.cpp
new file mode 100644
--- /dev/null
+++ b/Engine/Source/Programs/AutoRTFMTests/Private/UtilsTests.cpp
@@ -0,0 +1,257 @@
+// Copyright Epic Games, Inc. All Rights Reserved.
+
+#include "AutoRTFM/Utils.h"
+#include "Catch2Includes.h"
+
+TEST_CASE("AlignDown")
+{
+	REQUIRE(AutoRTFM::AlignDown(0, 1) == 0);
+	REQUIRE(AutoRTFM::AlignDown(1, 1) == 1);
+	REQUIRE(AutoRTFM::AlignDown(2, 1) == 2);
+	REQUIRE(AutoRTFM::AlignDown(3, 1) == 3);
+	REQUIRE(AutoRTFM::AlignDown(4, 1) == 4);
+	REQUIRE(AutoRTFM::AlignDown(5, 1) == 5);
+	REQUIRE(AutoRTFM::AlignDown(6, 1) == 6);
+	REQUIRE(AutoRTFM::AlignDown(7, 1) == 7);
+	REQUIRE(AutoRTFM::AlignDown(8, 1) == 8);
+
+	REQUIRE(AutoRTFM::AlignDown(0, 2) == 0);
+	REQUIRE(AutoRTFM::AlignDown(1, 2) == 0);
+	REQUIRE(AutoRTFM::AlignDown(2, 2) == 2);
+	REQUIRE(AutoRTFM::AlignDown(3, 2) == 2);
+	REQUIRE(AutoRTFM::AlignDown(4, 2) == 4);
+	REQUIRE(AutoRTFM::AlignDown(5, 2) == 4);
+	REQUIRE(AutoRTFM::AlignDown(6, 2) == 6);
+	REQUIRE(AutoRTFM::AlignDown(7, 2) == 6);
+	REQUIRE(AutoRTFM::AlignDown(8, 2) == 8);
+
+	REQUIRE(AutoRTFM::AlignDown(0, 4) == 0);
+	REQUIRE(AutoRTFM::AlignDown(1, 4) == 0);
+	REQUIRE(AutoRTFM::AlignDown(2, 4) == 0);
+	REQUIRE(AutoRTFM::AlignDown(3, 4) == 0);
+	REQUIRE(AutoRTFM::AlignDown(4, 4) == 4);
+	REQUIRE(AutoRTFM::AlignDown(5, 4) == 4);
+	REQUIRE(AutoRTFM::AlignDown(6, 4) == 4);
+	REQUIRE(AutoRTFM::AlignDown(7, 4) == 4);
+	REQUIRE(AutoRTFM::AlignDown(8, 4) == 8);
+
+	REQUIRE(AutoRTFM::AlignDown(0, 8) == 0);
+	REQUIRE(AutoRTFM::AlignDown(1, 8) == 0);
+	REQUIRE(AutoRTFM::AlignDown(2, 8) == 0);
+	REQUIRE(AutoRTFM::AlignDown(3, 8) == 0);
+	REQUIRE(AutoRTFM::AlignDown(4, 8) == 0);
+	REQUIRE(AutoRTFM::AlignDown(5, 8) == 0);
+	REQUIRE(AutoRTFM::AlignDown(6, 8) == 0);
+	REQUIRE(AutoRTFM::AlignDown(7, 8) == 0);
+	REQUIRE(AutoRTFM::AlignDown(8, 8) == 8);
+}
+
+TEST_CASE("AlignUp")
+{
+	REQUIRE(AutoRTFM::AlignUp(0, 1) == 0);
+	REQUIRE(AutoRTFM::AlignUp(1, 1) == 1);
+	REQUIRE(AutoRTFM::AlignUp(2, 1) == 2);
+	REQUIRE(AutoRTFM::AlignUp(3, 1) == 3);
+	REQUIRE(AutoRTFM::AlignUp(4, 1) == 4);
+	REQUIRE(AutoRTFM::AlignUp(5, 1) == 5);
+	REQUIRE(AutoRTFM::AlignUp(6, 1) == 6);
+	REQUIRE(AutoRTFM::AlignUp(7, 1) == 7);
+	REQUIRE(AutoRTFM::AlignUp(8, 1) == 8);
+
+	REQUIRE(AutoRTFM::AlignUp(0, 2) == 0);
+	REQUIRE(AutoRTFM::AlignUp(1, 2) == 2);
+	REQUIRE(AutoRTFM::AlignUp(2, 2) == 2);
+	REQUIRE(AutoRTFM::AlignUp(3, 2) == 4);
+	REQUIRE(AutoRTFM::AlignUp(4, 2) == 4);
+	REQUIRE(AutoRTFM::AlignUp(5, 2) == 6);
+	REQUIRE(AutoRTFM::AlignUp(6, 2) == 6);
+	REQUIRE(AutoRTFM::AlignUp(7, 2) == 8);
+	REQUIRE(AutoRTFM::AlignUp(8, 2) == 8);
+
+	REQUIRE(AutoRTFM::AlignUp(0, 4) == 0);
+	REQUIRE(AutoRTFM::AlignUp(1, 4) == 4);
+	REQUIRE(AutoRTFM::AlignUp(2, 4) == 4);
+	REQUIRE(AutoRTFM::AlignUp(3, 4) == 4);
+	REQUIRE(AutoRTFM::AlignUp(4, 4) == 4);
+	REQUIRE(AutoRTFM::AlignUp(5, 4) == 8);
+	REQUIRE(AutoRTFM::AlignUp(6, 4) == 8);
+	REQUIRE(AutoRTFM::AlignUp(7, 4) == 8);
+	REQUIRE(AutoRTFM::AlignUp(8, 4) == 8);
+
+	REQUIRE(AutoRTFM::AlignUp(0, 8) == 0);
+	REQUIRE(AutoRTFM::AlignUp(1, 8) == 8);
+	REQUIRE(AutoRTFM::AlignUp(2, 8) == 8);
+	REQUIRE(AutoRTFM::AlignUp(3, 8) == 8);
+	REQUIRE(AutoRTFM::AlignUp(4, 8) == 8);
+	REQUIRE(AutoRTFM::AlignUp(5, 8) == 8);
+	REQUIRE(AutoRTFM::AlignUp(6, 8) == 8);
+	REQUIRE(AutoRTFM::AlignUp(7, 8) == 8);
+	REQUIRE(AutoRTFM::AlignUp(8, 8) == 8);
+}
+
+
+TEST_CASE("RoundDown")
+{
+	REQUIRE(AutoRTFM::RoundDown(0, 1) == 0);
+	REQUIRE(AutoRTFM::RoundDown(1, 1) == 1);
+	REQUIRE(AutoRTFM::RoundDown(2, 1) == 2);
+	REQUIRE(AutoRTFM::RoundDown(3, 1) == 3);
+	REQUIRE(AutoRTFM::RoundDown(4, 1) == 4);
+	REQUIRE(AutoRTFM::RoundDown(5, 1) == 5);
+	REQUIRE(AutoRTFM::RoundDown(6, 1) == 6);
+	REQUIRE(AutoRTFM::RoundDown(7, 1) == 7);
+	REQUIRE(AutoRTFM::RoundDown(8, 1) == 8);
+
+	REQUIRE(AutoRTFM::RoundDown(0, 2) == 0);
+	REQUIRE(AutoRTFM::RoundDown(1, 2) == 0);
+	REQUIRE(AutoRTFM::RoundDown(2, 2) == 2);
+	REQUIRE(AutoRTFM::RoundDown(3, 2) == 2);
+	REQUIRE(AutoRTFM::RoundDown(4, 2) == 4);
+	REQUIRE(AutoRTFM::RoundDown(5, 2) == 4);
+	REQUIRE(AutoRTFM::RoundDown(6, 2) == 6);
+	REQUIRE(AutoRTFM::RoundDown(7, 2) == 6);
+	REQUIRE(AutoRTFM::RoundDown(8, 2) == 8);
+
+	REQUIRE(AutoRTFM::RoundDown(0, 3) == 0);
+	REQUIRE(AutoRTFM::RoundDown(1, 3) == 0);
+	REQUIRE(AutoRTFM::RoundDown(2, 3) == 0);
+	REQUIRE(AutoRTFM::RoundDown(3, 3) == 3);
+	REQUIRE(AutoRTFM::RoundDown(4, 3) == 3);
+	REQUIRE(AutoRTFM::RoundDown(5, 3) == 3);
+	REQUIRE(AutoRTFM::RoundDown(6, 3) == 6);
+	REQUIRE(AutoRTFM::RoundDown(7, 3) == 6);
+	REQUIRE(AutoRTFM::RoundDown(8, 3) == 6);
+
+	REQUIRE(AutoRTFM::RoundDown(0, 4) == 0);
+	REQUIRE(AutoRTFM::RoundDown(1, 4) == 0);
+	REQUIRE(AutoRTFM::RoundDown(2, 4) == 0);
+	REQUIRE(AutoRTFM::RoundDown(3, 4) == 0);
+	REQUIRE(AutoRTFM::RoundDown(4, 4) == 4);
+	REQUIRE(AutoRTFM::RoundDown(5, 4) == 4);
+	REQUIRE(AutoRTFM::RoundDown(6, 4) == 4);
+	REQUIRE(AutoRTFM::RoundDown(7, 4) == 4);
+	REQUIRE(AutoRTFM::RoundDown(8, 4) == 8);
+
+	REQUIRE(AutoRTFM::RoundDown(0, 5) == 0);
+	REQUIRE(AutoRTFM::RoundDown(1, 5) == 0);
+	REQUIRE(AutoRTFM::RoundDown(2, 5) == 0);
+	REQUIRE(AutoRTFM::RoundDown(3, 5) == 0);
+	REQUIRE(AutoRTFM::RoundDown(4, 5) == 0);
+	REQUIRE(AutoRTFM::RoundDown(5, 5) == 5);
+	REQUIRE(AutoRTFM::RoundDown(6, 5) == 5);
+	REQUIRE(AutoRTFM::RoundDown(7, 5) == 5);
+	REQUIRE(AutoRTFM::RoundDown(8, 5) == 5);
+
+	REQUIRE(AutoRTFM::RoundDown(0, 6) == 0);
+	REQUIRE(AutoRTFM::RoundDown(1, 6) == 0);
+	REQUIRE(AutoRTFM::RoundDown(2, 6) == 0);
+	REQUIRE(AutoRTFM::RoundDown(3, 6) == 0);
+	REQUIRE(AutoRTFM::RoundDown(4, 6) == 0);
+	REQUIRE(AutoRTFM::RoundDown(5, 6) == 0);
+	REQUIRE(AutoRTFM::RoundDown(6, 6) == 6);
+	REQUIRE(AutoRTFM::RoundDown(7, 6) == 6);
+	REQUIRE(AutoRTFM::RoundDown(8, 6) == 6);
+
+	REQUIRE(AutoRTFM::RoundDown(0, 7) == 0);
+	REQUIRE(AutoRTFM::RoundDown(1, 7) == 0);
+	REQUIRE(AutoRTFM::RoundDown(2, 7) == 0);
+	REQUIRE(AutoRTFM::RoundDown(3, 7) == 0);
+	REQUIRE(AutoRTFM::RoundDown(4, 7) == 0);
+	REQUIRE(AutoRTFM::RoundDown(5, 7) == 0);
+	REQUIRE(AutoRTFM::RoundDown(6, 7) == 0);
+	REQUIRE(AutoRTFM::RoundDown(7, 7) == 7);
+	REQUIRE(AutoRTFM::RoundDown(8, 7) == 7);
+
+	REQUIRE(AutoRTFM::RoundDown(0, 8) == 0);
+	REQUIRE(AutoRTFM::RoundDown(1, 8) == 0);
+	REQUIRE(AutoRTFM::RoundDown(2, 8) == 0);
+	REQUIRE(AutoRTFM::RoundDown(3, 8) == 0);
+	REQUIRE(AutoRTFM::RoundDown(4, 8) == 0);
+	REQUIRE(AutoRTFM::RoundDown(5, 8) == 0);
+	REQUIRE(AutoRTFM::RoundDown(6, 8) == 0);
+	REQUIRE(AutoRTFM::RoundDown(7, 8) == 0);
+	REQUIRE(AutoRTFM::RoundDown(8, 8) == 8);
+}
+
+TEST_CASE("RoundUp")
+{
+	REQUIRE(AutoRTFM::RoundUp(0, 1) == 0);
+	REQUIRE(AutoRTFM::RoundUp(1, 1) == 1);
+	REQUIRE(AutoRTFM::RoundUp(2, 1) == 2);
+	REQUIRE(AutoRTFM::RoundUp(3, 1) == 3);
+	REQUIRE(AutoRTFM::RoundUp(4, 1) == 4);
+	REQUIRE(AutoRTFM::RoundUp(5, 1) == 5);
+	REQUIRE(AutoRTFM::RoundUp(6, 1) == 6);
+	REQUIRE(AutoRTFM::RoundUp(7, 1) == 7);
+	REQUIRE(AutoRTFM::RoundUp(8, 1) == 8);
+
+	REQUIRE(AutoRTFM::RoundUp(0, 2) == 0);
+	REQUIRE(AutoRTFM::RoundUp(1, 2) == 2);
+	REQUIRE(AutoRTFM::RoundUp(2, 2) == 2);
+	REQUIRE(AutoRTFM::RoundUp(3, 2) == 4);
+	REQUIRE(AutoRTFM::RoundUp(4, 2) == 4);
+	REQUIRE(AutoRTFM::RoundUp(5, 2) == 6);
+	REQUIRE(AutoRTFM::RoundUp(6, 2) == 6);
+	REQUIRE(AutoRTFM::RoundUp(7, 2) == 8);
+	REQUIRE(AutoRTFM::RoundUp(8, 2) == 8);
+
+	REQUIRE(AutoRTFM::RoundUp(0, 3) == 0);
+	REQUIRE(AutoRTFM::RoundUp(1, 3) == 3);
+	REQUIRE(AutoRTFM::RoundUp(2, 3) == 3);
+	REQUIRE(AutoRTFM::RoundUp(3, 3) == 3);
+	REQUIRE(AutoRTFM::RoundUp(4, 3) == 6);
+	REQUIRE(AutoRTFM::RoundUp(5, 3) == 6);
+	REQUIRE(AutoRTFM::RoundUp(6, 3) == 6);
+	REQUIRE(AutoRTFM::RoundUp(7, 3) == 9);
+	REQUIRE(AutoRTFM::RoundUp(8, 3) == 9);
+
+	REQUIRE(AutoRTFM::RoundUp(0, 4) == 0);
+	REQUIRE(AutoRTFM::RoundUp(1, 4) == 4);
+	REQUIRE(AutoRTFM::RoundUp(2, 4) == 4);
+	REQUIRE(AutoRTFM::RoundUp(3, 4) == 4);
+	REQUIRE(AutoRTFM::RoundUp(4, 4) == 4);
+	REQUIRE(AutoRTFM::RoundUp(5, 4) == 8);
+	REQUIRE(AutoRTFM::RoundUp(6, 4) == 8);
+	REQUIRE(AutoRTFM::RoundUp(7, 4) == 8);
+	REQUIRE(AutoRTFM::RoundUp(8, 4) == 8);
+
+	REQUIRE(AutoRTFM::RoundUp(0, 5) == 0);
+	REQUIRE(AutoRTFM::RoundUp(1, 5) == 5);
+	REQUIRE(AutoRTFM::RoundUp(2, 5) == 5);
+	REQUIRE(AutoRTFM::RoundUp(3, 5) == 5);
+	REQUIRE(AutoRTFM::RoundUp(4, 5) == 5);
+	REQUIRE(AutoRTFM::RoundUp(5, 5) == 5);
+	REQUIRE(AutoRTFM::RoundUp(6, 5) == 10);
+	REQUIRE(AutoRTFM::RoundUp(7, 5) == 10);
+	REQUIRE(AutoRTFM::RoundUp(8, 5) == 10);
+
+	REQUIRE(AutoRTFM::RoundUp(0, 6) == 0);
+	REQUIRE(AutoRTFM::RoundUp(1, 6) == 6);
+	REQUIRE(AutoRTFM::RoundUp(2, 6) == 6);
+	REQUIRE(AutoRTFM::RoundUp(3, 6) == 6);
+	REQUIRE(AutoRTFM::RoundUp(4, 6) == 6);
+	REQUIRE(AutoRTFM::RoundUp(5, 6) == 6);
+	REQUIRE(AutoRTFM::RoundUp(6, 6) == 6);
+	REQUIRE(AutoRTFM::RoundUp(7, 6) == 12);
+	REQUIRE(AutoRTFM::RoundUp(8, 6) == 12);
+
+	REQUIRE(AutoRTFM::RoundUp(0, 7) == 0);
+	REQUIRE(AutoRTFM::RoundUp(1, 7) == 7);
+	REQUIRE(AutoRTFM::RoundUp(2, 7) == 7);
+	REQUIRE(AutoRTFM::RoundUp(3, 7) == 7);
+	REQUIRE(AutoRTFM::RoundUp(4, 7) == 7);
+	REQUIRE(AutoRTFM::RoundUp(5, 7) == 7);
+	REQUIRE(AutoRTFM::RoundUp(6, 7) == 7);
+	REQUIRE(AutoRTFM::RoundUp(7, 7) == 7);
+	REQUIRE(AutoRTFM::RoundUp(8, 7) == 14);
+
+	REQUIRE(AutoRTFM::RoundUp(0, 8) == 0);
+	REQUIRE(AutoRTFM::RoundUp(1, 8) == 8);
+	REQUIRE(AutoRTFM::RoundUp(2, 8) == 8);
+	REQUIRE(AutoRTFM::RoundUp(3, 8) == 8);
+	REQUIRE(AutoRTFM::RoundUp(4, 8) == 8);
+	REQUIRE(AutoRTFM::RoundUp(5, 8) == 8);
+	REQUIRE(AutoRTFM::RoundUp(6, 8) == 8);
+	REQUIRE(AutoRTFM::RoundUp(7, 8) == 8);
+	REQUIRE(AutoRTFM::RoundUp(8, 8) == 8);
+}
diff --git a/Engine/Source/Programs/AutoRTFMTests/Private/WriteLogTests.cpp b/Engine/Source/Programs/AutoRTFMTests/Private/WriteLogTests.cpp
--- a/Engine/Source/Programs/AutoRTFMTests/Private/WriteLogTests.cpp
+++ b/Engine/Source/Programs/AutoRTFMTests/Private/WriteLogTests.cpp
@@ -2,7 +2,6 @@
 
 #include "AutoRTFM/WriteLog.h"
 #include "Catch2Includes.h"
-#include "HAL/Platform.h"
 #include <algorithm>
 #include <cstddef>
 #include <numeric>
@@ -217,3 +216,240 @@ TEST_CASE("FWriteLog")
 		Check(RandomSize(/* NumEntries */ 32));
 	}
 }
+
+TEST_CASE("FWriteLog.Hash")
+{
+	struct FMemorySpan
+	{
+		std::byte* Data;
+		size_t Size;
+	};
+
+	// Populates Buffer with some pseudo-random data.
+	auto FillBuffer = [](FMemorySpan Buffer)
+	{
+		uint32_t Value = 976187;
+		for (size_t Offset = 0; Offset < Buffer.Size; Offset++)
+		{
+			Value = ((Value >> 3) * 921563) ^ ((Value << 1) * 743917);
+			Buffer.Data[Offset] = std::byte(Value & 255);
+		}
+	};
+
+	// Returns a FWriteLog populated with write log records of size WriteSize,
+	// and with a mis-alignment of AlignmentOffset. Records will have at least
+	// one byte between them.
+	// Writes is populated with all the memory spans for the write records.
+	// Gaps is populated with all the memory spans between the write records.
+	auto CreateWriteLog = [](
+		FMemorySpan Buffer,
+		size_t WriteSize,
+		size_t AlignmentOffset,
+		std::vector<FMemorySpan>& Writes,
+		std::vector<FMemorySpan>& Gaps)
+	{
+		AutoRTFM::FWriteLog WriteLog;
+		size_t Offset = AlignmentOffset;
+		while (Offset + WriteSize < Buffer.Size)
+		{
+			AutoRTFM::FWriteLogEntry Entry{};
+			Entry.LogicalAddress = Buffer.Data + Offset;
+			Entry.Data = Entry.LogicalAddress;
+			Entry.Size = WriteSize;
+			WriteLog.Push(Entry);
+
+			Writes.push_back(FMemorySpan{Buffer.Data + Offset, WriteSize});
+
+			Offset += WriteSize;
+
+			const size_t GapStart = Offset;
+			if (AlignmentOffset == 0)
+			{
+				// Round up offset to a multiple of WriteSize, but ensuring
+				// there's a gap of at least one byte between write log records.
+				// This ensures that records are not folded together.
+				Offset = AutoRTFM::RoundUp(Offset + 1, WriteSize);
+			}
+			else
+			{
+				// Round up offset to a multiple of WriteSize, then add the
+				// alignment offset to make the offset misaligned.
+				Offset = AutoRTFM::RoundUp(Offset, WriteSize) + AlignmentOffset;
+			}
+			const size_t GapEnd = std::min(Offset, Buffer.Size);
+			Gaps.push_back(FMemorySpan{Buffer.Data + GapStart, GapEnd - GapStart});
+		}
+		return WriteLog;
+	};
+
+	// Tests
+	{
+		constexpr size_t BufferSize = 1 << 13; // 8KiB
+		auto BufferData = std::unique_ptr<std::byte[]>(new std::byte[BufferSize]);
+		REQUIRE((reinterpret_cast<uintptr_t>(BufferData.get()) & 15) == 0);
+		FMemorySpan Buffer{BufferData.get(), BufferSize};
+		
+		FillBuffer(Buffer);
+
+		auto Check = [&](size_t WriteSize, size_t AlignmentOffset)
+		{
+			std::vector<FMemorySpan> Writes, Gaps;
+			AutoRTFM::FWriteLog WriteLog = CreateWriteLog(Buffer, WriteSize, AlignmentOffset, Writes, Gaps);
+
+			const AutoRTFM::FWriteLog::FHash OriginalHash = WriteLog.Hash(WriteLog.Num());
+
+			constexpr size_t NumChecks = 100;
+			size_t NumChanged = 0;
+			for (size_t I = 0; I < NumChecks; I++)
+			{
+				// Pick a random write.
+				FMemorySpan& Write = Writes[(I * 16831) % Writes.size()];
+				// Pick a random byte within the write, change it and check the hash changes.
+				size_t Offset = (I * 838483) % Write.Size;
+				const std::byte OriginalByte = Write.Data[Offset];
+				Write.Data[Offset] = ~Write.Data[Offset];
+				const AutoRTFM::FWriteLog::FHash NewHash = WriteLog.Hash(WriteLog.Num());
+				NumChanged += (OriginalHash != NewHash) ? 1 : 0;
+				Write.Data[Offset] = OriginalByte;
+			}
+
+			// Expect at least 95% of the byte changes to have affected the hash.
+			REQUIRE(NumChanged > (NumChecks * 95 / 100));
+			REQUIRE(OriginalHash == WriteLog.Hash(WriteLog.Num()));
+			
+			// Modify all the padding bytes between write records. These should
+			// not affect the hash.
+			for (FMemorySpan& Gap : Gaps)
+			{
+				for (size_t I = 0; I < Gap.Size; I++)
+				{
+					Gap.Data[I] = ~Gap.Data[I];
+				}
+			}
+			REQUIRE(OriginalHash == WriteLog.Hash(WriteLog.Num()));
+		};
+
+		Check(/* WriteSize */ 1, /* AlignmentOffset */ 0);
+		Check(/* WriteSize */ 2, /* AlignmentOffset */ 0);
+		Check(/* WriteSize */ 2, /* AlignmentOffset */ 1);
+		Check(/* WriteSize */ 3, /* AlignmentOffset */ 0);
+		Check(/* WriteSize */ 3, /* AlignmentOffset */ 1);
+		Check(/* WriteSize */ 3, /* AlignmentOffset */ 2);
+		Check(/* WriteSize */ 4, /* AlignmentOffset */ 0);
+		Check(/* WriteSize */ 4, /* AlignmentOffset */ 1);
+		Check(/* WriteSize */ 4, /* AlignmentOffset */ 2);
+		Check(/* WriteSize */ 4, /* AlignmentOffset */ 3);
+		Check(/* WriteSize */ 5, /* AlignmentOffset */ 0);
+		Check(/* WriteSize */ 5, /* AlignmentOffset */ 1);
+		Check(/* WriteSize */ 5, /* AlignmentOffset */ 3);
+		Check(/* WriteSize */ 5, /* AlignmentOffset */ 4);
+		Check(/* WriteSize */ 6, /* AlignmentOffset */ 0);
+		Check(/* WriteSize */ 6, /* AlignmentOffset */ 2);
+		Check(/* WriteSize */ 6, /* AlignmentOffset */ 3);
+		Check(/* WriteSize */ 6, /* AlignmentOffset */ 5);
+		Check(/* WriteSize */ 7, /* AlignmentOffset */ 0);
+		Check(/* WriteSize */ 7, /* AlignmentOffset */ 1);
+		Check(/* WriteSize */ 7, /* AlignmentOffset */ 4);
+		Check(/* WriteSize */ 7, /* AlignmentOffset */ 6);
+		Check(/* WriteSize */ 8, /* AlignmentOffset */ 0);
+		Check(/* WriteSize */ 8, /* AlignmentOffset */ 2);
+		Check(/* WriteSize */ 8, /* AlignmentOffset */ 3);
+		Check(/* WriteSize */ 8, /* AlignmentOffset */ 7);
+		Check(/* WriteSize */ 9, /* AlignmentOffset */ 0);
+		Check(/* WriteSize */ 9, /* AlignmentOffset */ 1);
+		Check(/* WriteSize */ 9, /* AlignmentOffset */ 5);
+		Check(/* WriteSize */ 9, /* AlignmentOffset */ 8);
+		Check(/* WriteSize */ 10, /* AlignmentOffset */ 0);
+		Check(/* WriteSize */ 10, /* AlignmentOffset */ 3);
+		Check(/* WriteSize */ 10, /* AlignmentOffset */ 5);
+		Check(/* WriteSize */ 10, /* AlignmentOffset */ 9);
+		Check(/* WriteSize */ 15, /* AlignmentOffset */ 0);
+		Check(/* WriteSize */ 15, /* AlignmentOffset */ 6);
+		Check(/* WriteSize */ 15, /* AlignmentOffset */ 9);
+		Check(/* WriteSize */ 15, /* AlignmentOffset */ 14);
+		Check(/* WriteSize */ 16, /* AlignmentOffset */ 0);
+		Check(/* WriteSize */ 16, /* AlignmentOffset */ 8);
+		Check(/* WriteSize */ 16, /* AlignmentOffset */ 9);
+		Check(/* WriteSize */ 16, /* AlignmentOffset */ 15);
+		Check(/* WriteSize */ 32, /* AlignmentOffset */ 0);
+		Check(/* WriteSize */ 32, /* AlignmentOffset */ 8);
+		Check(/* WriteSize */ 32, /* AlignmentOffset */ 16);
+		Check(/* WriteSize */ 32, /* AlignmentOffset */ 31);
+		Check(/* WriteSize */ 64, /* AlignmentOffset */ 0);
+		Check(/* WriteSize */ 64, /* AlignmentOffset */ 7);
+		Check(/* WriteSize */ 64, /* AlignmentOffset */ 30);
+		Check(/* WriteSize */ 64, /* AlignmentOffset */ 63);
+		Check(/* WriteSize */ 100, /* AlignmentOffset */ 0);
+		Check(/* WriteSize */ 100, /* AlignmentOffset */ 1);
+		Check(/* WriteSize */ 100, /* AlignmentOffset */ 10);
+		Check(/* WriteSize */ 100, /* AlignmentOffset */ 99);
+	}
+
+	{
+		constexpr size_t BufferSize = 1 << 20; // 1MiB
+		auto BufferData = std::unique_ptr<std::byte[]>(new std::byte[BufferSize]);
+		REQUIRE((reinterpret_cast<uintptr_t>(BufferData.get()) & 15) == 0);
+		FMemorySpan Buffer{BufferData.get(), BufferSize};
+
+		FillBuffer(Buffer);
+
+		auto Benchmark = [&](Catch::Benchmark::Chronometer Meter, size_t WriteSize, size_t AlignmentOffset)
+		{
+			std::vector<FMemorySpan> Writes, Gaps;
+			AutoRTFM::FWriteLog WriteLog = CreateWriteLog(Buffer, WriteSize, AlignmentOffset, Writes, Gaps);
+			Meter.measure([&]
+			{
+				return WriteLog.Hash(WriteLog.Num());
+			});
+		};
+
+		SECTION("Aligned")
+		{
+#define BENCH(WRITE_SIZE) \
+			BENCHMARK_ADVANCED("WriteSize: " #WRITE_SIZE)(Catch::Benchmark::Chronometer Meter) \
+			{ \
+				Benchmark(Meter, /* WriteSize */ WRITE_SIZE, /* AlignmentOffset */ 0); \
+			};
+			BENCH(1);
+			BENCH(2);
+			BENCH(3);
+			BENCH(4);
+			BENCH(5);
+			BENCH(6);
+			BENCH(7);
+			BENCH(8);
+			BENCH(9);
+			BENCH(10);
+			BENCH(15);
+			BENCH(16);
+			BENCH(32);
+			BENCH(64);
+			BENCH(128);
+#undef BENCH
+		}
+		SECTION("Unaligned")
+		{
+#define BENCH(WRITE_SIZE) \
+			BENCHMARK_ADVANCED("WriteSize: " #WRITE_SIZE)(Catch::Benchmark::Chronometer Meter) \
+			{ \
+				Benchmark(Meter, /* WriteSize */ WRITE_SIZE, /* AlignmentOffset */ 1); \
+			};
+			BENCH(1);
+			BENCH(2);
+			BENCH(3);
+			BENCH(4);
+			BENCH(5);
+			BENCH(6);
+			BENCH(7);
+			BENCH(8);
+			BENCH(9);
+			BENCH(10);
+			BENCH(15);
+			BENCH(16);
+			BENCH(32);
+			BENCH(64);
+			BENCH(128);
+#undef BENCH
+		}
+	}
+}
diff --git a/Engine/Source/Runtime/Core/Private/AutoRTFM/Transaction.cpp b/Engine/Source/Runtime/Core/Private/AutoRTFM/Transaction.cpp
--- a/Engine/Source/Runtime/Core/Private/AutoRTFM/Transaction.cpp
+++ b/Engine/Source/Runtime/Core/Private/AutoRTFM/Transaction.cpp
@@ -6,9 +6,6 @@
 #include "ContextStatus.h"
 #include "CallNestInlines.h"
 #include "GlobalData.h"
-#include "HAL/Platform.h"
-#include "Hash/xxhash.h"
-#include "Misc/AssertionMacros.h"
 #include "TransactionInlines.h"
 #include "HitSet.h"
 #include "Utils.h"
@@ -474,7 +471,7 @@ void FTransaction::ValidateWriteHash() const
 	static constexpr const char Message[] =
 		"Memory modified in a transaction was also modified in an call to AutoRTFM::Open(). "
 		"This may lead to memory corruption if the transaction is aborted.";
-	if (UNLIKELY(OldHash != NewHash))
+	if (AUTORTFM_UNLIKELY(OldHash != NewHash))
 	{
 		if (CurrentMemoryValidationLevel == EMemoryValidationLevel::Warn)
 		{
@@ -501,82 +498,15 @@ FTransaction::FWriteHash FTransaction::CalculateNestedWriteHash() const
 
 FTransaction::FWriteHash FTransaction::CalculateNestedWriteHashWithLimit(size_t NumWriteEntries, const void *OpenReturnAddress) const
 {
-	FXxHash64Builder Hasher;
+	FWriteHash Hash = 0;
 	if (nullptr != Parent)
 	{
-		FWriteHash ParentHash = Parent->CalculateNestedWriteHashWithLimit(Parent->WriteLog.Num(), OpenReturnAddress);
-		Hasher.Update(&ParentHash, sizeof(ParentHash));
+		Hash = 31 * Parent->CalculateNestedWriteHashWithLimit(Parent->WriteLog.Num(), OpenReturnAddress);
 	}
-
-	FOpenHashThrottler::FHashScope Profile(Context->GetOpenHashThrottler(), OpenReturnAddress, WriteLog);
-
-	size_t WriteIndex = 0;
-
-	using FSmallValue = uint64;
-	FSmallValue SmallValue = 0;
-	for(auto Iter = WriteLog.begin(); Iter != WriteLog.end(); ++Iter, ++WriteIndex)
 	{
-		if (WriteIndex == NumWriteEntries)
-		{
-			break;
-		}
-
-		const FWriteLogEntry& Entry = *Iter;
-		if (UNLIKELY(Entry.bNoMemoryValidation))
-		{
-			continue; // Next write
-		}
-
-		if (Entry.Size > sizeof(FSmallValue))
-		{
-			Hasher.Update(Entry.LogicalAddress, Entry.Size);
-		}
-
-		const uintptr_t Address = reinterpret_cast<uintptr_t>(Entry.LogicalAddress);
-		switch (Entry.Size)
-		{
-			case 8:
-				if (LIKELY((Address & 7) == 0))
-				{
-					uint64 Load = *reinterpret_cast<uint64 const*>(Entry.LogicalAddress);
-					SmallValue = (SmallValue * 31) ^ static_cast<FSmallValue>(Load);
-				}
-				break;
-
-			case 4:
-				if (LIKELY((Address & 3) == 0))
-				{
-					uint32 Load = *reinterpret_cast<uint32 const*>(Entry.LogicalAddress);
-					SmallValue = (SmallValue * 31) ^ static_cast<FSmallValue>(Load);
-				}
-				break;
-
-			case 2:
-				if (LIKELY((Address & 1) == 0))
-				{
-					uint16 Load = *reinterpret_cast<uint16 const*>(Entry.LogicalAddress);
-					SmallValue = (SmallValue * 31) ^ static_cast<FSmallValue>(Load);
-				}
-				break;
-
-			case 1:
-				SmallValue = (SmallValue * 31) ^ static_cast<FSmallValue>(Entry.LogicalAddress[0]);
-				break;
-
-			default:
-				// Unaligned / non-POT write.
-				for (size_t I = 0; I < Entry.Size; I++)
-				{
-					SmallValue = (SmallValue * 31) ^ static_cast<FSmallValue>(Entry.LogicalAddress[I]);
-				}
-				break;
-		}
-
+		FOpenHashThrottler::FHashScope Profile(Context->GetOpenHashThrottler(), OpenReturnAddress, WriteLog);
+		Hash ^= WriteLog.Hash(NumWriteEntries);
 	}
-	Hasher.Update(&SmallValue, sizeof(SmallValue));
-
-	auto Hash = Hasher.Finalize().Hash;
-	static_assert(std::is_same_v<decltype(Hash), FWriteHash>, "FWriteHash type mismatch");
 	return Hash;
 }
 
diff --git a/Engine/Source/Runtime/Core/Private/AutoRTFM/Transaction.h b/Engine/Source/Runtime/Core/Private/AutoRTFM/Transaction.h
--- a/Engine/Source/Runtime/Core/Private/AutoRTFM/Transaction.h
+++ b/Engine/Source/Runtime/Core/Private/AutoRTFM/Transaction.h
@@ -156,7 +156,7 @@ public:
 	inline bool IsOnStack(const void* LogicalAddress) const;
 
 private:
-	using FWriteHash = uint64;
+	using FWriteHash = FWriteLog::FHash;
 	using FTaskArray = TTaskArray<TTask<void()>>;
 
 	void Undo();
diff --git a/Engine/Source/Runtime/Core/Private/AutoRTFM/Utils.h b/Engine/Source/Runtime/Core/Private/AutoRTFM/Utils.h
--- a/Engine/Source/Runtime/Core/Private/AutoRTFM/Utils.h
+++ b/Engine/Source/Runtime/Core/Private/AutoRTFM/Utils.h
@@ -2,6 +2,7 @@
 
 #pragma once
 
+#include <type_traits>
 #if (defined(__AUTORTFM) && __AUTORTFM)
 
 #include "AutoRTFM/AutoRTFM.h"
@@ -103,6 +104,42 @@ inline void EnsureFailure(const char* File, int Line, const char* Condition, con
 	va_end(Args);
 }
 
+// Rounds Value up to the next multiple of Alignment.
+// Alignment must be a power of two.
+template<typename T>
+inline constexpr T AlignDown(T Value, T Alignment)
+{
+	static_assert(std::is_integral_v<T>);
+	return Value & ~(Alignment - 1);
+}
+
+// Rounds Value down to the next multiple of Alignment.
+// Alignment must be a power of two.
+template<typename T>
+inline constexpr T AlignUp(T Value, T Alignment)
+{
+	static_assert(std::is_integral_v<T>);
+	return (Value + Alignment - 1) & ~(Alignment - 1);
+}
+
+// Rounds Value up to the next multiple of Multiple.
+// Unlike AlignDown(), Multiple does not need to be a power of two.
+template<typename T>
+inline constexpr T RoundDown(T Value, T Multiple)
+{
+	static_assert(std::is_integral_v<T>);
+	return (Value / Multiple) * Multiple;
+}
+
+// Rounds Value down to the next multiple of Multiple.
+// Unlike AlignUp(), Multiple does not need to be a power of two.
+template<typename T>
+inline constexpr T RoundUp(T Value, T Multiple)
+{
+	static_assert(std::is_integral_v<T>);
+	return RoundDown(Value + Multiple - 1, Multiple);
+}
+
 } // namespace AutoRTFM
 
 #define AUTORTFM_LIKELY(x) __builtin_expect(!!(x), 1)
diff --git a/Engine/Source/Runtime/Core/Private/AutoRTFM/WriteLog.cpp b/Engine/Source/Runtime/Core/Private/AutoRTFM/WriteLog.cpp
new file mode 100644
--- /dev/null
+++ b/Engine/Source/Runtime/Core/Private/AutoRTFM/WriteLog.cpp
@@ -0,0 +1,87 @@
+// Copyright Epic Games, Inc. All Rights Reserved.
+
+#if (defined(__AUTORTFM) && __AUTORTFM)
+
+#include "WriteLog.h"
+#include "Hash/xxhash.h"
+
+namespace AutoRTFM
+{
+
+FWriteLog::FHash FWriteLog::Hash(size_t NumWriteEntries) const
+{
+	FXxHash64Builder Hasher;
+
+	size_t WriteIndex = 0;
+
+	using FSmallValue = uint64;
+	FSmallValue SmallValue = 0;
+
+	for(auto Iter = begin(); Iter != end(); ++Iter, ++WriteIndex)
+	{
+		if (WriteIndex == NumWriteEntries)
+		{
+			break;
+		}
+
+		const FWriteLogEntry& Entry = *Iter;
+		if (UNLIKELY(Entry.bNoMemoryValidation))
+		{
+			continue; // Next write
+		}
+
+		if (Entry.Size > sizeof(FSmallValue))
+		{
+			Hasher.Update(Entry.LogicalAddress, Entry.Size);
+		}
+
+		const uintptr_t Address = reinterpret_cast<uintptr_t>(Entry.LogicalAddress);
+		switch (Entry.Size)
+		{
+			case 8:
+				if (LIKELY((Address & 7) == 0))
+				{
+					uint64 Load = *reinterpret_cast<uint64 const*>(Entry.LogicalAddress);
+					SmallValue = (SmallValue * 31) ^ static_cast<FSmallValue>(Load);
+					continue;
+				}
+				break;
+
+			case 4:
+				if (LIKELY((Address & 3) == 0))
+				{
+					uint32 Load = *reinterpret_cast<uint32 const*>(Entry.LogicalAddress);
+					SmallValue = (SmallValue * 31) ^ static_cast<FSmallValue>(Load);
+					continue;
+				}
+				break;
+
+			case 2:
+				if (LIKELY((Address & 1) == 0))
+				{
+					uint16 Load = *reinterpret_cast<uint16 const*>(Entry.LogicalAddress);
+					SmallValue = (SmallValue * 31) ^ static_cast<FSmallValue>(Load);
+					continue;
+				}
+				break;
+
+			case 1:
+				SmallValue = (SmallValue * 31) ^ static_cast<FSmallValue>(Entry.LogicalAddress[0]);
+				continue;
+		}
+
+		// Unaligned / non-POT write.
+		for (size_t I = 0; I < Entry.Size; I++)
+		{
+			SmallValue = (SmallValue * 31) ^ static_cast<FSmallValue>(Entry.LogicalAddress[I]);
+		}
+	}
+	
+	Hasher.Update(&SmallValue, sizeof(SmallValue));
+
+	return Hasher.Finalize().Hash;
+}
+
+} // namespace AutoRTFM
+
+#endif // (defined(__AUTORTFM) && __AUTORTFM)
diff --git a/Engine/Source/Runtime/Core/Private/AutoRTFM/WriteLog.h b/Engine/Source/Runtime/Core/Private/AutoRTFM/WriteLog.h
--- a/Engine/Source/Runtime/Core/Private/AutoRTFM/WriteLog.h
+++ b/Engine/Source/Runtime/Core/Private/AutoRTFM/WriteLog.h
@@ -7,7 +7,6 @@
 #include "AutoRTFM/AutoRTFM.h"
 #include "HAL/Platform.h"
 #include "Misc/CoreMiscDefines.h"
-#include "Templates/AlignmentTemplates.h"
 #include "Utils.h"
 #include <cstddef>
 #include <cstdint>
@@ -218,7 +217,7 @@ namespace AutoRTFM
 			}
 			else if (UNLIKELY(PushResult == FBlock::EPushResult::Full))
 			{
-				const size_t RequiredSize = Align(sizeof(FBlock) + Entry.Size, alignof(FRecord)) + sizeof(FRecord);
+				const size_t RequiredSize = AlignUp(sizeof(FBlock) + Entry.Size, alignof(FRecord)) + sizeof(FRecord);
 				FBlock* NewBlock = FBlock::Allocate(std::max(RequiredSize, FBlock::DefaultSize));
 				NewBlock->PrevBlock = TailBlock;
 				TailBlock->NextBlock = NewBlock;
@@ -369,6 +368,12 @@ namespace AutoRTFM
 		// Return the total size in bytes for all entries in the log.
 		UE_AUTORTFM_FORCEINLINE size_t TotalSize() const { return TotalSizeBytes; }
 
+		// Returns a hash of the first NumWriteEntries entries' logical memory
+		// tracked by the write log. This is the memory post-write, not the
+		// original memory that would be restored on abort.
+		using FHash = uint64_t;
+		FHash Hash(size_t NumWriteEntries) const;
+
 	private:
 		template<size_t SIZE>
 		static constexpr bool IsAlignedForTRecord = (SIZE & (alignof(FRecord) - 1)) == 0;
