[AutoRTFM] Add a transactionally-safe FArchive writer class.

This class wraps an inner FArchive, which must be passed-in, and defers any writes to the archive while inside a transaction.

This allows us to fix the Network Profiler by simply wrapping the FileWriter archive and then using it normally.

diff --git a/Engine/Source/Programs/AutoRTFMTests/Private/ArchiveWriterTests.cpp b/Engine/Source/Programs/AutoRTFMTests/Private/ArchiveWriterTests.cpp
new file mode 100644
--- /dev/null
+++ b/Engine/Source/Programs/AutoRTFMTests/Private/ArchiveWriterTests.cpp
@@ -0,0 +1,309 @@
+// Copyright Epic Games, Inc. All Rights Reserved.
+
+#include "Catch2Includes.h"
+#include "AutoRTFM/AutoRTFM.h"
+#include "AutoRTFM/AutoRTFMTesting.h"
+#include "Serialization/MemoryWriter.h"
+#include "Serialization/TransactionallySafeArchiveWriter.h"
+
+/**
+ * This test class wraps an FMemoryWriter, but check-fails if Serialize is ever called
+ * from within a transaction. This will cause our tests to fail if we violate the
+ * transactional boundary.
+ */
+class FCheckingMemoryWriter : public FMemoryWriter
+{
+	using Super = FMemoryWriter;
+
+public:
+	using Super::Super;
+
+	virtual FString GetArchiveName() const override
+	{
+		return TEXT("FCheckingMemoryWriter");
+	}
+
+	virtual void Serialize(void* Data, int64 Num) override
+	{
+		check(!AutoRTFM::IsTransactional());
+		Super::Serialize(Data, Num);
+	}
+};
+
+static void SerializeData(FArchive& Archive)
+{
+	// Serializes various data into the archive: short, packed int, bool, bit-array, string.
+	uint16 Short = 123;
+	Archive << Short;
+	uint32 PackedUint = 4567;
+	Archive.SerializeIntPacked(PackedUint);
+	bool Boolean = true;
+	Archive << Boolean;
+	uint32 Bits = 0x55555555;
+	Archive.SerializeBits(&Bits, 11);
+	FString("Hello").SerializeAsANSICharArray(Archive);
+}
+
+TEST_CASE("TransactionallySafeArchiveWriter.MatchesMemoryWriter")
+{
+	// This test verifies that a variety of serialization methods (via SerializeData) work as designed.
+
+	TArray<uint8> NormalStorage;
+	FMemoryWriter MemoryWriter(NormalStorage);
+	SerializeData(MemoryWriter);
+
+	SECTION("Write")
+	{
+		TArray<uint8> TSStorage;
+		FTransactionallySafeArchiveWriter SafeWriter(MakeUnique<FCheckingMemoryWriter>(TSStorage));
+		SerializeData(SafeWriter);
+
+		REQUIRE(TSStorage == NormalStorage);
+	}
+
+	SECTION("Ctor, Commit(Write)")
+	{
+		TArray<uint8> TSStorage;
+		FTransactionallySafeArchiveWriter SafeWriter(MakeUnique<FCheckingMemoryWriter>(TSStorage));
+
+		AutoRTFM::Testing::Commit([&]
+		{
+			SerializeData(SafeWriter);
+
+			// The FTransactionallySafeArchiveWriter must defer writes into the backing storage.
+			REQUIRE(TSStorage.IsEmpty());
+		});
+
+		REQUIRE(TSStorage == NormalStorage);
+	}
+
+	SECTION("Commit(Ctor, Write)")
+	{
+		TArray<uint8> TSStorage;
+
+		AutoRTFM::Testing::Commit([&]
+		{
+			FTransactionallySafeArchiveWriter SafeWriter(MakeUnique<FCheckingMemoryWriter>(TSStorage));
+			SerializeData(SafeWriter);
+		});
+
+		REQUIRE(TSStorage == NormalStorage);
+	}
+}
+
+TEST_CASE("TransactionallySafeArchiveWriter.CommitAndAbortWork")
+{
+	// This test verifies that different patterns of construction, committing, and aborting
+	// all yield the expected result.
+
+	TArray<uint8> NormalStorage;
+	FMemoryWriter MemoryWriter(NormalStorage);
+	FString("Hello").SerializeAsANSICharArray(MemoryWriter);
+
+	SECTION("Ctor, Write, Abort(Write)")
+	{
+		TArray<uint8> TSStorage;
+		FTransactionallySafeArchiveWriter SafeWriter(MakeUnique<FCheckingMemoryWriter>(TSStorage));
+
+		FString("Hello").SerializeAsANSICharArray(SafeWriter);
+
+		AutoRTFM::Testing::Abort([&]
+		{
+			FString("World").SerializeAsANSICharArray(SafeWriter);
+			AutoRTFM::AbortTransaction();
+		});
+
+		REQUIRE(TSStorage == NormalStorage);
+	}
+
+	SECTION("Ctor, Commit(Write, Abort(Write))")
+	{
+		TArray<uint8> TSStorage;
+		FTransactionallySafeArchiveWriter SafeWriter(MakeUnique<FCheckingMemoryWriter>(TSStorage));
+
+		AutoRTFM::Testing::Commit([&]
+		{
+			FString("Hello").SerializeAsANSICharArray(SafeWriter);
+
+			AutoRTFM::Testing::Abort([&]
+			{
+				FString("World").SerializeAsANSICharArray(SafeWriter);
+				AutoRTFM::AbortTransaction();
+			});
+		});
+
+		REQUIRE(TSStorage == NormalStorage);
+	}
+
+	SECTION("Commit(Ctor, Write, Abort(Write))")
+	{
+		TArray<uint8> TSStorage;
+
+		AutoRTFM::Testing::Commit([&]
+		{
+			FTransactionallySafeArchiveWriter SafeWriter(MakeUnique<FCheckingMemoryWriter>(TSStorage));
+			FString("Hello").SerializeAsANSICharArray(SafeWriter);
+
+			AutoRTFM::Testing::Abort([&]
+			{
+				FString("World").SerializeAsANSICharArray(SafeWriter);
+				AutoRTFM::AbortTransaction();
+			});
+		});
+
+		REQUIRE(TSStorage == NormalStorage);
+	}
+
+	SECTION("Ctor, Commit(Abort(Write), Write)")
+	{
+		TArray<uint8> TSStorage;
+		FTransactionallySafeArchiveWriter SafeWriter(MakeUnique<FCheckingMemoryWriter>(TSStorage));
+
+		AutoRTFM::Testing::Commit([&]
+		{
+			AutoRTFM::Testing::Abort([&]
+			{
+				FString("World").SerializeAsANSICharArray(SafeWriter);
+				AutoRTFM::AbortTransaction();
+			});
+
+			FString("Hello").SerializeAsANSICharArray(SafeWriter);
+		});
+
+		REQUIRE(TSStorage == NormalStorage);
+	}
+
+	SECTION("Commit(Ctor, Abort(Write), Write)")
+	{
+		TArray<uint8> TSStorage;
+
+		AutoRTFM::Testing::Commit([&]
+		{
+			FTransactionallySafeArchiveWriter SafeWriter(MakeUnique<FCheckingMemoryWriter>(TSStorage));
+
+			AutoRTFM::Testing::Abort([&]
+			{
+				FString("World").SerializeAsANSICharArray(SafeWriter);
+				AutoRTFM::AbortTransaction();
+			});
+
+			FString("Hello").SerializeAsANSICharArray(SafeWriter);
+		});
+
+		REQUIRE(TSStorage == NormalStorage);
+	}
+
+	SECTION("Commit(Ctor, Abort(Write), Write)")
+	{
+		TArray<uint8> TSStorage;
+
+		AutoRTFM::Testing::Commit([&]
+		{
+			FTransactionallySafeArchiveWriter SafeWriter(MakeUnique<FCheckingMemoryWriter>(TSStorage));
+
+			AutoRTFM::Testing::Abort([&]
+			{
+				FString("World").SerializeAsANSICharArray(SafeWriter);
+				AutoRTFM::AbortTransaction();
+			});
+
+			FString("Hello").SerializeAsANSICharArray(SafeWriter);
+		});
+
+		REQUIRE(TSStorage == NormalStorage);
+	}
+
+	SECTION("Ctor, Write, Abort(Commit(Write)))")
+	{
+		TArray<uint8> TSStorage;
+		FTransactionallySafeArchiveWriter SafeWriter(MakeUnique<FCheckingMemoryWriter>(TSStorage));
+
+		FString("Hello").SerializeAsANSICharArray(SafeWriter);
+
+		AutoRTFM::Testing::Abort([&]
+		{
+			AutoRTFM::Testing::Commit([&]
+			{
+				FString("World").SerializeAsANSICharArray(SafeWriter);
+			});
+			AutoRTFM::AbortTransaction();
+		});
+
+		REQUIRE(TSStorage == NormalStorage);
+	}
+}
+
+TEST_CASE("TransactionallySafeArchiveWriter.MixingWrites")
+{
+	// This test verifies that an archive writes can be mixed between the transactionally-safe
+	// writer and its wrapped archive, using Release() to take the archive back as needed.
+
+	TArray<uint8> NormalStorage;
+	FMemoryWriter MemoryWriter(NormalStorage);
+	FString("Hello").SerializeAsANSICharArray(MemoryWriter);
+	FString("World").SerializeAsANSICharArray(MemoryWriter);
+
+	SECTION("Wrap, Write, Release, Write")
+	{
+		TArray<uint8> TSStorage;
+		FTransactionallySafeArchiveWriter SafeWriter(MakeUnique<FCheckingMemoryWriter>(TSStorage));
+		FString("Hello").SerializeAsANSICharArray(SafeWriter);
+
+		TUniquePtr<FArchive> ReleasedWriter = SafeWriter.Release();
+		FString("World").SerializeAsANSICharArray(*ReleasedWriter);
+
+		REQUIRE(TSStorage == NormalStorage);
+	}
+
+	SECTION("Wrap, Commit(Write), Release, Write")
+	{
+		TArray<uint8> TSStorage;
+		FTransactionallySafeArchiveWriter SafeWriter(MakeUnique<FCheckingMemoryWriter>(TSStorage));
+
+		AutoRTFM::Testing::Commit([&]
+		{
+			FString("Hello").SerializeAsANSICharArray(SafeWriter);
+		});
+
+		TUniquePtr<FArchive> ReleasedWriter = SafeWriter.Release();
+		FString("World").SerializeAsANSICharArray(*ReleasedWriter);
+
+		REQUIRE(TSStorage == NormalStorage);
+	}
+
+	SECTION("Write, Wrap, Commit(Write), Release")
+	{
+		TArray<uint8> TSStorage;
+		TUniquePtr<FCheckingMemoryWriter> CheckingWriter(MakeUnique<FCheckingMemoryWriter>(TSStorage));
+
+		FString("Hello").SerializeAsANSICharArray(*CheckingWriter);
+
+		FTransactionallySafeArchiveWriter SafeWriter(MoveTemp(CheckingWriter));
+
+		AutoRTFM::Testing::Commit([&]
+		{
+			FString("World").SerializeAsANSICharArray(SafeWriter);
+		});
+
+		REQUIRE(TSStorage == NormalStorage);
+
+		SafeWriter.Release();
+
+		REQUIRE(TSStorage == NormalStorage);
+	}
+
+	SECTION("Write, Wrap, Release, Write")
+	{
+		TArray<uint8> TSStorage;
+		TUniquePtr<FCheckingMemoryWriter> CheckingWriter(MakeUnique<FCheckingMemoryWriter>(TSStorage));
+
+		FString("Hello").SerializeAsANSICharArray(*CheckingWriter);
+
+		FTransactionallySafeArchiveWriter SafeWriter(MoveTemp(CheckingWriter));
+		TUniquePtr<FArchive> ReleasedWriter = SafeWriter.Release();
+
+		FString("World").SerializeAsANSICharArray(*ReleasedWriter);
+
+		REQUIRE(TSStorage == NormalStorage);
+	}
+}
diff --git a/Engine/Source/Runtime/Core/Public/Serialization/Archive.h b/Engine/Source/Runtime/Core/Public/Serialization/Archive.h
--- a/Engine/Source/Runtime/Core/Public/Serialization/Archive.h
+++ b/Engine/Source/Runtime/Core/Public/Serialization/Archive.h
@@ -1638,7 +1638,7 @@ public:
 #if defined(__clang_analyzer__)
 		// Suppress core.uninitialized.Assign static analysis warning
 		// Base class FArchive::Serialize() does not modify the value being serialized, so the analysis
-		// assumes any data being serialized to could be uninitalized, but practically speaking any
+		// assumes any data being serialized to could be uninitialized, but practically speaking any
 		// data being serialized will use a derived class that always modifies the value.
 		memset(V, 0, Length);
 #endif
diff --git a/Engine/Source/Runtime/Core/Public/Serialization/TransactionallySafeArchiveWriter.h b/Engine/Source/Runtime/Core/Public/Serialization/TransactionallySafeArchiveWriter.h
new file mode 100644
--- /dev/null
+++ b/Engine/Source/Runtime/Core/Public/Serialization/TransactionallySafeArchiveWriter.h
@@ -0,0 +1,111 @@
+// Copyright Epic Games, Inc. All Rights Reserved.
+
+#pragma once
+
+#include "AutoRTFM/AutoRTFM.h"
+#include "CoreTypes.h"
+#include "Containers/Array.h"
+#include "Containers/UnrealString.h"
+#include "Misc/AssertionMacros.h"
+#include "Serialization/Archive.h"
+#include "Templates/UniquePtr.h"
+
+/**
+ * FTransactionallySafeArchiveWriter takes ownership of an existing FArchive and allows it to
+ * be written to during a transaction. This works by deferring file writes into a memory buffer
+ * when inside a transaction. When the transaction is committed, the buffer is written into the
+ * passed-in FArchive.
+ * 
+ * Outside of a transaction, writes are passed through to the wrapped archive transparently.
+ * 
+ * This class only supports basic archive functionality like << and Serialize. Other operations
+ * like Seek and Tell are not supported.
+ */
+class FTransactionallySafeArchiveWriter : public FArchive
+{
+public:
+	FTransactionallySafeArchiveWriter(TUniquePtr<FArchive> Ar) : InnerArchive(MoveTemp(Ar))
+	{
+		check(InnerArchive);
+		check(InnerArchive->IsSaving());
+		SetIsSaving(true);
+	}
+
+	~FTransactionallySafeArchiveWriter()
+	{
+		if (InnerArchive && bRegisteredCommitHandler)
+		{
+			// Our wrapper is going away, so move everything into a lambda-capture for safekeeping.
+			AutoRTFM::PopOnCommitHandler(this);
+			AutoRTFM::PushOnCommitHandler(this, [Archive = InnerArchive.Release(), DataToSerialize = MoveTemp(DeferredWrites)]() mutable
+			{
+				Archive->Serialize(DataToSerialize.GetData(), DataToSerialize.NumBytes() );
+				delete Archive;
+			});
+		}
+	}
+
+	/** 
+	 * If we are outside of a transaction, you can have your archive back.
+	 * (It's dangerous to allow Release within a transaction since this opens the door to out-of-sequence writes.)
+	 */
+	TUniquePtr<FArchive> Release()
+	{
+		check(!AutoRTFM::IsTransactional());
+		return MoveTemp(InnerArchive);
+	}
+
+	virtual FString GetArchiveName() const override
+	{ 
+		return InnerArchive->GetArchiveName();
+	}
+
+	virtual void Seek(int64 InPos) final
+	{
+		unimplemented();
+	}
+
+	virtual int64 Tell() final
+	{
+		unimplemented();
+		return 0;
+	}
+
+	virtual void Serialize(void* Data, int64 Num) override
+	{
+		check(InnerArchive);
+
+		if (AutoRTFM::IsTransactional())
+		{
+			DeferredWrites.Append(static_cast<uint8*>(Data), IntCastChecked<DeferredWriteBuffer::SizeType>(Num));
+			MaybeRegisterCommitHandler();
+		}
+		else
+		{
+			check(DeferredWrites.IsEmpty());
+			InnerArchive->Serialize(Data, Num);
+		}
+	}
+
+private:
+	void MaybeRegisterCommitHandler()
+	{
+		if (!bRegisteredCommitHandler)
+		{
+			bRegisteredCommitHandler = true;
+
+			AutoRTFM::PushOnCommitHandler(this, [this]
+			{
+				InnerArchive->Serialize(DeferredWrites.GetData(), DeferredWrites.NumBytes());
+				DeferredWrites.Reset();
+				bRegisteredCommitHandler = false;
+			});
+		}
+	}
+
+	using DeferredWriteBuffer = TArray<uint8>;
+
+	TUniquePtr<FArchive> InnerArchive;
+	DeferredWriteBuffer DeferredWrites;
+	bool bRegisteredCommitHandler = false;
+};
\ No newline at end of file
diff --git a/Engine/Source/Runtime/Engine/Private/NetworkProfiler.cpp b/Engine/Source/Runtime/Engine/Private/NetworkProfiler.cpp
--- a/Engine/Source/Runtime/Engine/Private/NetworkProfiler.cpp
+++ b/Engine/Source/Runtime/Engine/Private/NetworkProfiler.cpp
@@ -12,6 +12,7 @@
 #include "Misc/App.h"
 #include "Engine/World.h"
 #include "TimerManager.h"
+#include "Serialization/TransactionallySafeArchiveWriter.h"
 #include "UObject/UnrealType.h"
 
 #if USE_NETWORK_PROFILER
@@ -756,7 +757,7 @@ void FNetworkProfiler::TrackSessionChange( bool bShouldContinueTracking, const F
 			const FString FinalFileName = FPaths::ProfilingDir() + FileName;
 
 			IFileManager::Get().MakeDirectory( *FPaths::GetPath( FinalFileName ) );
-			FileWriter = IFileManager::Get().CreateFileWriter( *FinalFileName, FILEWRITE_EvenIfReadOnly );
+			FileWriter = new FTransactionallySafeArchiveWriter( TUniquePtr<FArchive>( IFileManager::Get().CreateFileWriter( *FinalFileName, FILEWRITE_EvenIfReadOnly )));
 			check( FileWriter );
 			
 			UE_LOG(LogNet, Log, TEXT("Network Profiler: Creating session file at location %s"), *FinalFileName);
