[AutoRTFM] Fix assertion in TransactionallySafeArchiveWriter when OnCommit is used.

diff --git a/Engine/Source/Programs/AutoRTFMTests/Private/ArchiveWriterTests.cpp b/Engine/Source/Programs/AutoRTFMTests/Private/ArchiveWriterTests.cpp
--- a/Engine/Source/Programs/AutoRTFMTests/Private/ArchiveWriterTests.cpp
+++ b/Engine/Source/Programs/AutoRTFMTests/Private/ArchiveWriterTests.cpp
@@ -307,3 +307,75 @@ TEST_CASE("TransactionallySafeArchiveWriter.MixingWrites")
 		REQUIRE(TSStorage == NormalStorage);
 	}
 }
+
+TEST_CASE("TransactionallySafeArchiveWriter.OnCommitWorks")
+{
+	TArray<uint8> NormalStorage;
+	FMemoryWriter MemoryWriter(NormalStorage);
+	FString("One").SerializeAsANSICharArray(MemoryWriter);
+	FString("Two").SerializeAsANSICharArray(MemoryWriter);
+	FString("Three").SerializeAsANSICharArray(MemoryWriter);
+
+	SECTION("Write, Commit(Write, OnCommit)")
+	{
+		TArray<uint8> TSStorage;
+		FTransactionallySafeArchiveWriter SafeWriter(MakeUnique<FCheckingMemoryWriter>(TSStorage));
+
+		FString("One").SerializeAsANSICharArray(SafeWriter);
+
+		AutoRTFM::Testing::Commit([&]
+		{
+			FString("Two").SerializeAsANSICharArray(SafeWriter);
+
+			AutoRTFM::OnCommit([&]
+			{
+				FString("Three").SerializeAsANSICharArray(SafeWriter);
+			});
+		});
+
+		REQUIRE(TSStorage == NormalStorage);
+	}
+
+	SECTION("Commit(OnCommit, Write), Write")
+	{
+		TArray<uint8> TSStorage;
+		FTransactionallySafeArchiveWriter SafeWriter(MakeUnique<FCheckingMemoryWriter>(TSStorage));
+
+		AutoRTFM::Testing::Commit([&]
+		{
+			AutoRTFM::OnCommit([&]
+			{
+				FString("Two").SerializeAsANSICharArray(SafeWriter);
+			});
+
+			FString("One").SerializeAsANSICharArray(SafeWriter);
+		});
+
+		FString("Three").SerializeAsANSICharArray(SafeWriter);
+
+		REQUIRE(TSStorage == NormalStorage);
+	}
+
+	SECTION("Commit(OnCommit, Write, OnCommit)")
+	{
+		TArray<uint8> TSStorage;
+		FTransactionallySafeArchiveWriter SafeWriter(MakeUnique<FCheckingMemoryWriter>(TSStorage));
+
+		AutoRTFM::Testing::Commit([&]
+		{
+			AutoRTFM::OnCommit([&]
+			{
+				FString("Two").SerializeAsANSICharArray(SafeWriter);
+			});
+
+			FString("One").SerializeAsANSICharArray(SafeWriter);
+
+			AutoRTFM::OnCommit([&]
+			{
+				FString("Three").SerializeAsANSICharArray(SafeWriter);
+			});
+		});
+
+		REQUIRE(TSStorage == NormalStorage);
+	}
+}
\ No newline at end of file
diff --git a/Engine/Source/Runtime/Core/Public/Serialization/TransactionallySafeArchiveWriter.h b/Engine/Source/Runtime/Core/Public/Serialization/TransactionallySafeArchiveWriter.h
--- a/Engine/Source/Runtime/Core/Public/Serialization/TransactionallySafeArchiveWriter.h
+++ b/Engine/Source/Runtime/Core/Public/Serialization/TransactionallySafeArchiveWriter.h
@@ -75,7 +75,7 @@ public:
 	{
 		check(InnerArchive);
 
-		if (AutoRTFM::IsTransactional())
+		if (bRegisteredCommitHandler || AutoRTFM::IsTransactional())
 		{
 			DeferredWrites.Append(static_cast<uint8*>(Data), IntCastChecked<DeferredWriteBuffer::SizeType>(Num));
 			MaybeRegisterCommitHandler();
