[AutoRTFM] Rename AutoRTFMTests/Private/Utils.h to AutoRTFMTestUtils.h

Once AutoRTFM is extracted from core, the "AutoRTFM/" part of includes will be dropped and then this becomes ambiguous with the private AutoRTFM Utils.h.

diff --git a/Engine/Source/Programs/AutoRTFMTests/Private/Abort.cpp b/Engine/Source/Programs/AutoRTFMTests/Private/Abort.cpp
--- a/Engine/Source/Programs/AutoRTFMTests/Private/Abort.cpp
+++ b/Engine/Source/Programs/AutoRTFMTests/Private/Abort.cpp
@@ -1,8 +1,8 @@
 // Copyright Epic Games, Inc. All Rights Reserved.
 
 #include "AutoRTFM/AutoRTFMTesting.h"
+#include "AutoRTFMTestUtils.h"
 #include "Catch2Includes.h"
-#include "Utils.h"
 
 #include <map>
 #include <vector>
diff --git a/Engine/Source/Programs/AutoRTFMTests/Private/AutoRTFMTestUtils.h b/Engine/Source/Programs/AutoRTFMTests/Private/AutoRTFMTestUtils.h
new file mode 100644
--- /dev/null
+++ b/Engine/Source/Programs/AutoRTFMTests/Private/AutoRTFMTestUtils.h
@@ -0,0 +1,152 @@
+// Copyright Epic Games, Inc. All Rights Reserved.
+
+#pragma once
+
+#include "AutoRTFM/AutoRTFM.h"
+#include "Misc/FeedbackContext.h"
+
+#include <vector>
+
+namespace AutoRTFMTestUtils
+{
+
+// Temporarily changes the AutoRTFM retry mode for the lifetime of the FScopedRetry object.
+struct FScopedRetry
+{
+	FScopedRetry(AutoRTFM::ForTheRuntime::EAutoRTFMRetryTransactionState NewRetry)
+		: OldRetry(AutoRTFM::ForTheRuntime::GetRetryTransaction())
+	{
+		AutoRTFM::ForTheRuntime::SetRetryTransaction(NewRetry);
+	}
+
+	~FScopedRetry()
+	{
+		AutoRTFM::ForTheRuntime::SetRetryTransaction(OldRetry);
+	}
+
+	AutoRTFM::ForTheRuntime::EAutoRTFMRetryTransactionState OldRetry;
+};
+
+// A helper class that for the lifetime of the object, intercepts and records UE_LOG warnings.
+class FCaptureWarningContext : private FFeedbackContext
+{
+public:
+	FCaptureWarningContext() : OldContext(GWarn) { GWarn = this; }
+	~FCaptureWarningContext() { GWarn = OldContext; }
+
+	void Serialize(const TCHAR* V, ELogVerbosity::Type Verbosity, const FName& Category) override
+	{
+		if (Verbosity == ELogVerbosity::Warning)
+		{
+			Warnings.push_back(V);
+		}
+		else
+		{
+			OldContext->Serialize(V, Verbosity, Category);
+		}
+	}
+
+	void Serialize(const TCHAR* V, ELogVerbosity::Type Verbosity, const FName& Category, double Time) override
+	{
+		if (Verbosity == ELogVerbosity::Warning)
+		{
+			Warnings.push_back(V);
+		}
+		else
+		{
+			OldContext->Serialize(V, Verbosity, Category, Time);
+		}
+	}
+
+	const std::vector<FString>& GetWarnings() const
+	{
+		return Warnings;
+	}
+
+private:
+	FCaptureWarningContext(const FCaptureWarningContext&) = delete;
+	FCaptureWarningContext& operator = (const FCaptureWarningContext&) = delete;
+
+	FFeedbackContext* OldContext = nullptr;
+	std::vector<FString> Warnings;
+};
+
+// A helper class that temporarily changes the memory validation level for the lifetime of the
+// object, restoring the original level on destruction.
+class ScopedMemoryValidationLevel
+{
+public:
+	ScopedMemoryValidationLevel(AutoRTFM::EMemoryValidationLevel NewLevel)
+		: PrevLevel(AutoRTFM::ForTheRuntime::GetMemoryValidationLevel())
+	{
+		AutoRTFM::ForTheRuntime::SetMemoryValidationLevel(NewLevel);
+	}
+
+	~ScopedMemoryValidationLevel()
+	{
+		AutoRTFM::ForTheRuntime::SetMemoryValidationLevel(PrevLevel);
+	}
+private:
+	ScopedMemoryValidationLevel(const ScopedMemoryValidationLevel&) = delete;
+	ScopedMemoryValidationLevel& operator = (const ScopedMemoryValidationLevel&) = delete;
+
+	const AutoRTFM::EMemoryValidationLevel PrevLevel;
+};
+
+#define AUTORTFM_SCOPED_DISABLE_MEMORY_VALIDATION() \
+	AutoRTFMTestUtils::ScopedMemoryValidationLevel DisableMemoryValidation \
+		{AutoRTFM::EMemoryValidationLevel::Disabled}; \
+	static_assert(true) /* require semicolon */
+
+#define AUTORTFM_SCOPED_ENABLE_MEMORY_VALIDATION_AS_WARNING() \
+	AutoRTFMTestUtils::ScopedMemoryValidationLevel EnableMemoryValidationAsWarning \
+		{AutoRTFM::EMemoryValidationLevel::Warn}; \
+	static_assert(true) /* require semicolon */
+
+static constexpr const TCHAR* kMemoryModifiedWarning =
+    TEXT("Memory modified in a transaction was also modified in an call to AutoRTFM::Open(). ")
+    TEXT("This may lead to memory corruption if the transaction is aborted.");
+
+class FScopedEnsureOnInternalAbort final
+{
+public:
+	FScopedEnsureOnInternalAbort(const bool bState)
+		: bOriginal(AutoRTFM::ForTheRuntime::GetEnsureOnInternalAbort())
+	{
+		AutoRTFM::ForTheRuntime::SetEnsureOnInternalAbort(bState);
+	}
+
+	~FScopedEnsureOnInternalAbort()
+	{
+		AutoRTFM::ForTheRuntime::SetEnsureOnInternalAbort(bOriginal);
+	}
+
+private:
+	FScopedEnsureOnInternalAbort(const FScopedEnsureOnInternalAbort&) = delete;
+	FScopedEnsureOnInternalAbort& operator = (const FScopedEnsureOnInternalAbort&) = delete;
+
+	const bool bOriginal;
+};
+
+class FScopedInternalAbortAction final
+{
+public:
+	FScopedInternalAbortAction(AutoRTFM::ForTheRuntime::EAutoRTFMInternalAbortActionState State)
+		: Original(AutoRTFM::ForTheRuntime::GetInternalAbortAction())
+	{
+		AutoRTFM::ForTheRuntime::SetInternalAbortAction(State);
+	}
+
+	~FScopedInternalAbortAction()
+	{
+		AutoRTFM::ForTheRuntime::SetInternalAbortAction(Original);
+	}
+
+private:
+	FScopedInternalAbortAction(const FScopedInternalAbortAction&) = delete;
+	FScopedInternalAbortAction& operator = (const FScopedInternalAbortAction&) = delete;
+
+	const AutoRTFM::ForTheRuntime::EAutoRTFMInternalAbortActionState Original;
+};
+
+}
diff --git a/Engine/Source/Programs/AutoRTFMTests/Private/Delegate.cpp b/Engine/Source/Programs/AutoRTFMTests/Private/Delegate.cpp
--- a/Engine/Source/Programs/AutoRTFMTests/Private/Delegate.cpp
+++ b/Engine/Source/Programs/AutoRTFMTests/Private/Delegate.cpp
@@ -1,11 +1,11 @@
 // Copyright Epic Games, Inc. All Rights Reserved.
 
 #include "AutoRTFM/AutoRTFMTesting.h"
+#include "AutoRTFMTestUtils.h"
 #include "Catch2Includes.h"
 #include "Delegates/Delegate.h"
 #include "MyAutoRTFMTestObject.h"
 #include "UObject/GCObject.h"
-#include "Utils.h"
 
 TEST_CASE("Delegate.Broadcast")
 {
diff --git a/Engine/Source/Programs/AutoRTFMTests/Private/FakeVM.cpp b/Engine/Source/Programs/AutoRTFMTests/Private/FakeVM.cpp
--- a/Engine/Source/Programs/AutoRTFMTests/Private/FakeVM.cpp
+++ b/Engine/Source/Programs/AutoRTFMTests/Private/FakeVM.cpp
@@ -1,8 +1,8 @@
 // Copyright Epic Games, Inc. All Rights Reserved.
 
 #include "AutoRTFM/AutoRTFM.h"
+#include "AutoRTFMTestUtils.h"
 #include "Catch2Includes.h"
-#include "Utils.h"
 
 #include <vector>
 
diff --git a/Engine/Source/Programs/AutoRTFMTests/Private/Open.cpp b/Engine/Source/Programs/AutoRTFMTests/Private/Open.cpp
--- a/Engine/Source/Programs/AutoRTFMTests/Private/Open.cpp
+++ b/Engine/Source/Programs/AutoRTFMTests/Private/Open.cpp
@@ -1,13 +1,13 @@
 // Copyright Epic Games, Inc. All Rights Reserved.
 
 #include "AutoRTFM/AutoRTFM.h"
+#include "AutoRTFMTestUtils.h"
 #include "Catch2Includes.h"
 
 #include <atomic>
 #include <catch_amalgamated.hpp>
 #include <map>
 #include <vector>
-#include "Utils.h"
 
 namespace
 {
diff --git a/Engine/Source/Programs/AutoRTFMTests/Private/OpenAPI.cpp b/Engine/Source/Programs/AutoRTFMTests/Private/OpenAPI.cpp
--- a/Engine/Source/Programs/AutoRTFMTests/Private/OpenAPI.cpp
+++ b/Engine/Source/Programs/AutoRTFMTests/Private/OpenAPI.cpp
@@ -1,8 +1,8 @@
 // Copyright Epic Games, Inc. All Rights Reserved.
 
 #include "AutoRTFM/AutoRTFM.h"
+#include "AutoRTFMTestUtils.h"
 #include "Catch2Includes.h"
-#include "Utils.h"
 
 #include <atomic>
 #include <map>
diff --git a/Engine/Source/Programs/AutoRTFMTests/Private/RWLockTests.cpp b/Engine/Source/Programs/AutoRTFMTests/Private/RWLockTests.cpp
--- a/Engine/Source/Programs/AutoRTFMTests/Private/RWLockTests.cpp
+++ b/Engine/Source/Programs/AutoRTFMTests/Private/RWLockTests.cpp
@@ -3,9 +3,10 @@
 #include "Catch2Includes.h"
 #include "AutoRTFM/AutoRTFM.h"
 #include "AutoRTFM/AutoRTFMTesting.h"
+#include "AutoRTFMTestUtils.h"
 #include "Misc/TransactionallySafeRWLock.h"
 #include "Misc/ScopeRWLock.h"
-#include "Utils.h"
+
 #include <thread>
 #include <mutex>
 
diff --git a/Engine/Source/Programs/AutoRTFMTests/Private/UECoreTests.cpp b/Engine/Source/Programs/AutoRTFMTests/Private/UECoreTests.cpp
--- a/Engine/Source/Programs/AutoRTFMTests/Private/UECoreTests.cpp
+++ b/Engine/Source/Programs/AutoRTFMTests/Private/UECoreTests.cpp
@@ -1,6 +1,7 @@
 // Copyright Epic Games, Inc. All Rights Reserved.
 
 #include "AutoRTFM/AutoRTFM.h"
+#include "AutoRTFMTestUtils.h"
 #include "Catch2Includes.h"
 #include "Async/TransactionallySafeMutex.h"
 #include "Async/ParallelFor.h"
@@ -32,7 +33,7 @@
 #include "Misc/PackageName.h"
 #include "Misc/ScopeRWLock.h"
 #include "StructUtils/PropertyBag.h"
-#include "Utils.h"
+
 #include <thread>
 #include <mutex>
 
diff --git a/Engine/Source/Programs/AutoRTFMTests/Private/UEPlatform.cpp b/Engine/Source/Programs/AutoRTFMTests/Private/UEPlatform.cpp
--- a/Engine/Source/Programs/AutoRTFMTests/Private/UEPlatform.cpp
+++ b/Engine/Source/Programs/AutoRTFMTests/Private/UEPlatform.cpp
@@ -1,11 +1,12 @@
 // Copyright Epic Games, Inc. All Rights Reserved.
 
 #include "AutoRTFM/AutoRTFMTesting.h"
-#include "Catch2Includes.h"
 #include "AutoRTFM/AutoRTFM.h"
+#include "AutoRTFMTestUtils.h"
+#include "Catch2Includes.h"
 #include "HAL/PlatformMisc.h"
 #include "Misc/Guid.h"
-#include "Utils.h"
+
 
 TEST_CASE("FPlatformMisc.CreateGUID")
 {
diff --git a/Engine/Source/Programs/AutoRTFMTests/Private/Utils.h b/Engine/Source/Programs/AutoRTFMTests/Private/Utils.h
deleted file mode 100644
