[AutoRTFM] Fix IsAsyncLoading when transactional loading is queued.

Previously, `IsAsyncLoading()` would only return true if the wrapped loader was actually performing a load. This is not what callers expect. As soon as `LoadPackageAsync()` is called, callers expect `IsAsyncLoading()` to return true immediately.

In practice this triggers an ensure() inside `FStreamableManager::StreamInternal` but the map continues playing normally. Since it's ultimately harmless, AFAICS we don't need to do an integration request into 34.20 here.

diff --git a/Engine/Source/Programs/AutoRTFMTests/Private/UECoreTests.cpp b/Engine/Source/Programs/AutoRTFMTests/Private/UECoreTests.cpp
--- a/Engine/Source/Programs/AutoRTFMTests/Private/UECoreTests.cpp
+++ b/Engine/Source/Programs/AutoRTFMTests/Private/UECoreTests.cpp
@@ -2470,6 +2470,23 @@ TEST_CASE("UECore.AsyncLoading")
 		FlushAsyncLoading(RequestId);
 	}
 
+	SECTION("IsAsyncLoading")
+	{
+		int32 RequestId = -1;
+
+		AutoRTFM::Testing::Commit([&]
+		{
+			REQUIRE(!IsAsyncLoading());
+
+			FString Name(FString::Printf(TEXT("/AutoRTFMTestPackage%d"), __LINE__));
+			RequestId = LoadPackageAsync(Name);
+
+			REQUIRE(IsAsyncLoading());
+		});
+
+		FlushAsyncLoading(RequestId);
+	}
+
 	SECTION("FlushAsyncLoading")
 	{
 		FString Name(FString::Printf(TEXT("/AutoRTFMTestPackage%d"), __LINE__));
diff --git a/Engine/Source/Runtime/CoreUObject/Private/Serialization/TransactionallySafeAsyncLoading.cpp b/Engine/Source/Runtime/CoreUObject/Private/Serialization/TransactionallySafeAsyncLoading.cpp
--- a/Engine/Source/Runtime/CoreUObject/Private/Serialization/TransactionallySafeAsyncLoading.cpp
+++ b/Engine/Source/Runtime/CoreUObject/Private/Serialization/TransactionallySafeAsyncLoading.cpp
@@ -385,7 +385,7 @@ public:
 
 	bool IsAsyncLoadingPackages() override
 	{
-		return WrappedPackageLoader->IsAsyncLoadingPackages();
+		return !TransactionallyDeferredLoadPackages.IsEmpty() || WrappedPackageLoader->IsAsyncLoadingPackages();
 	}
 
 	void NotifyConstructedDuringAsyncLoading(UObject* Object, bool bSubObject) override
