[AutoRTFM] Add support for Flush to the FTransactionallySafeArchiveWriter.

Previously, Flush() was ignored. Now, it is forwarded on to the InnerArchive (and deferred as necessary).

diff --git a/Engine/Source/Programs/AutoRTFMTests/Private/ArchiveWriterTests.cpp b/Engine/Source/Programs/AutoRTFMTests/Private/ArchiveWriterTests.cpp
--- a/Engine/Source/Programs/AutoRTFMTests/Private/ArchiveWriterTests.cpp
+++ b/Engine/Source/Programs/AutoRTFMTests/Private/ArchiveWriterTests.cpp
@@ -3,6 +3,7 @@
 #include "Catch2Includes.h"
 #include "AutoRTFM.h"
 #include "AutoRTFMTesting.h"
+#include "Misc/NotNull.h"
 #include "Serialization/MemoryWriter.h"
 #include "Serialization/TransactionallySafeArchiveWriter.h"
 
@@ -16,7 +17,10 @@ class FCheckingMemoryWriter : public FMemoryWriter
 	using Super = FMemoryWriter;
 
 public:
-	using Super::Super;
+	FCheckingMemoryWriter(TArray<uint8>& Buffer, TNotNull<int*> InNumFlushes) : Super(Buffer), NumFlushes(*InNumFlushes)
+	{
+		NumFlushes = 0;
+	}
 
 	virtual FString GetArchiveName() const override
 	{
@@ -28,6 +32,13 @@ public:
 		check(!AutoRTFM::IsTransactional());
 		Super::Serialize(Data, Num);
 	}
+
+	virtual void Flush() override
+	{
+		++NumFlushes;
+	}
+	
+	int& NumFlushes;
 };
 
 static void SerializeData(FArchive& Archive)
@@ -55,16 +66,29 @@ TEST_CASE("TransactionallySafeArchiveWriter.MatchesMemoryWriter")
 	SECTION("Write")
 	{
 		TArray<uint8> TSStorage;
-		FTransactionallySafeArchiveWriter SafeWriter(MakeUnique<FCheckingMemoryWriter>(TSStorage));
+		int NumFlushes = 0;
+		FTransactionallySafeArchiveWriter SafeWriter(MakeUnique<FCheckingMemoryWriter>(TSStorage, &NumFlushes));
 		SerializeData(SafeWriter);
 
 		REQUIRE(TSStorage == NormalStorage);
+		REQUIRE(NumFlushes == 0);
+	}
+
+	SECTION("Flush")
+	{
+		TArray<uint8> TSStorage;
+		int NumFlushes = 0;
+		FTransactionallySafeArchiveWriter SafeWriter(MakeUnique<FCheckingMemoryWriter>(TSStorage, &NumFlushes));
+		SafeWriter.Flush();
+
+		REQUIRE(NumFlushes == 1);
 	}
 
 	SECTION("Ctor, Commit(Write)")
 	{
 		TArray<uint8> TSStorage;
-		FTransactionallySafeArchiveWriter SafeWriter(MakeUnique<FCheckingMemoryWriter>(TSStorage));
+		int NumFlushes = 0;
+		FTransactionallySafeArchiveWriter SafeWriter(MakeUnique<FCheckingMemoryWriter>(TSStorage, &NumFlushes));
 
 		AutoRTFM::Testing::Commit([&]
 		{
@@ -75,19 +99,103 @@ TEST_CASE("TransactionallySafeArchiveWriter.MatchesMemoryWriter")
 		});
 
 		REQUIRE(TSStorage == NormalStorage);
+		REQUIRE(NumFlushes == 0);
 	}
 
 	SECTION("Commit(Ctor, Write)")
 	{
 		TArray<uint8> TSStorage;
+		int NumFlushes = 0;
 
 		AutoRTFM::Testing::Commit([&]
 		{
-			FTransactionallySafeArchiveWriter SafeWriter(MakeUnique<FCheckingMemoryWriter>(TSStorage));
+			FTransactionallySafeArchiveWriter SafeWriter(MakeUnique<FCheckingMemoryWriter>(TSStorage, &NumFlushes));
 			SerializeData(SafeWriter);
+
+			// The FTransactionallySafeArchiveWriter must defer writes into the backing storage.
+			REQUIRE(TSStorage.IsEmpty());
 		});
 
 		REQUIRE(TSStorage == NormalStorage);
+		REQUIRE(NumFlushes == 0);
+	}
+
+	SECTION("Ctor, Commit(Flush)")
+	{
+		TArray<uint8> TSStorage;
+		int NumFlushes = 0;
+		FTransactionallySafeArchiveWriter SafeWriter(MakeUnique<FCheckingMemoryWriter>(TSStorage, &NumFlushes));
+
+		AutoRTFM::Testing::Commit([&]
+		{
+			SafeWriter.Flush();
+
+			// The FTransactionallySafeArchiveWriter must defer flushes on the backing storage.
+			REQUIRE(NumFlushes == 0);
+		});
+
+		REQUIRE(TSStorage.IsEmpty());
+		REQUIRE(NumFlushes == 1);
+	}
+
+	SECTION("Commit(Ctor, Flush, Flush, Flush)")
+	{
+		TArray<uint8> TSStorage;
+		int NumFlushes = 0;
+
+		AutoRTFM::Testing::Commit([&]
+		{
+			FTransactionallySafeArchiveWriter SafeWriter(MakeUnique<FCheckingMemoryWriter>(TSStorage, &NumFlushes));
+			SafeWriter.Flush();
+			SafeWriter.Flush();
+			SafeWriter.Flush();
+
+			// The FTransactionallySafeArchiveWriter must defer flushes on the backing storage.
+			REQUIRE(NumFlushes == 0);
+		});
+
+		REQUIRE(TSStorage.IsEmpty());
+		REQUIRE(NumFlushes == 1);
+	}
+
+	SECTION("Ctor, Commit(Write, Flush)")
+	{
+		TArray<uint8> TSStorage;
+		int NumFlushes = 0;
+		FTransactionallySafeArchiveWriter SafeWriter(MakeUnique<FCheckingMemoryWriter>(TSStorage, &NumFlushes));
+
+		AutoRTFM::Testing::Commit([&]
+		{
+			SerializeData(SafeWriter);
+			SafeWriter.Flush();
+
+			// The FTransactionallySafeArchiveWriter must defer writes and flushes to the backing storage.
+			REQUIRE(TSStorage.IsEmpty());
+			REQUIRE(NumFlushes == 0);
+		});
+
+		REQUIRE(TSStorage == NormalStorage);
+		REQUIRE(NumFlushes == 1);
+	}
+
+	SECTION("Commit(Ctor, Write, Flush)")
+	{
+		TArray<uint8> TSStorage;
+		int NumFlushes = 0;
+
+		AutoRTFM::Testing::Commit([&]
+		{
+			FTransactionallySafeArchiveWriter SafeWriter(MakeUnique<FCheckingMemoryWriter>(TSStorage, &NumFlushes));
+			SerializeData(SafeWriter);
+			SafeWriter.Flush();
+
+			// The FTransactionallySafeArchiveWriter must defer writes and flushes to the backing storage.
+			REQUIRE(TSStorage.IsEmpty());
+			REQUIRE(NumFlushes == 0);
+		});
+
+		REQUIRE(TSStorage == NormalStorage);
+		REQUIRE(NumFlushes == 1);
 	}
 }
 
@@ -100,26 +208,30 @@ TEST_CASE("TransactionallySafeArchiveWriter.CommitAndAbortWork")
 	FMemoryWriter MemoryWriter(NormalStorage);
 	FString("Hello").SerializeAsANSICharArray(MemoryWriter);
 
-	SECTION("Ctor, Write, Abort(Write)")
+	SECTION("Ctor, Write, Abort(Write, Flush)")
 	{
 		TArray<uint8> TSStorage;
-		FTransactionallySafeArchiveWriter SafeWriter(MakeUnique<FCheckingMemoryWriter>(TSStorage));
+		int NumFlushes = 0;
+		FTransactionallySafeArchiveWriter SafeWriter(MakeUnique<FCheckingMemoryWriter>(TSStorage, &NumFlushes));
 
 		FString("Hello").SerializeAsANSICharArray(SafeWriter);
 
 		AutoRTFM::Testing::Abort([&]
 		{
 			FString("World").SerializeAsANSICharArray(SafeWriter);
+			SafeWriter.Flush();
 			AutoRTFM::AbortTransaction();
 		});
 
 		REQUIRE(TSStorage == NormalStorage);
+		REQUIRE(NumFlushes == 0);
 	}
 
-	SECTION("Ctor, Commit(Write, Abort(Write))")
+	SECTION("Ctor, Commit(Write, Abort(Write), Abort(Flush))")
 	{
 		TArray<uint8> TSStorage;
-		FTransactionallySafeArchiveWriter SafeWriter(MakeUnique<FCheckingMemoryWriter>(TSStorage));
+		int NumFlushes = 0;
+		FTransactionallySafeArchiveWriter SafeWriter(MakeUnique<FCheckingMemoryWriter>(TSStorage, &NumFlushes));
 
 		AutoRTFM::Testing::Commit([&]
 		{
@@ -130,34 +242,45 @@ TEST_CASE("TransactionallySafeArchiveWriter.CommitAndAbortWork")
 				FString("World").SerializeAsANSICharArray(SafeWriter);
 				AutoRTFM::AbortTransaction();
 			});
+
+			AutoRTFM::Testing::Abort([&]
+			{
+				SafeWriter.Flush();
+				AutoRTFM::AbortTransaction();
+			});
 		});
 
 		REQUIRE(TSStorage == NormalStorage);
+		REQUIRE(NumFlushes == 0);
 	}
 
-	SECTION("Commit(Ctor, Write, Abort(Write))")
+	SECTION("Commit(Ctor, Write, Abort(Write, Flush))")
 	{
 		TArray<uint8> TSStorage;
+		int NumFlushes = 0;
 
 		AutoRTFM::Testing::Commit([&]
 		{
-			FTransactionallySafeArchiveWriter SafeWriter(MakeUnique<FCheckingMemoryWriter>(TSStorage));
+			FTransactionallySafeArchiveWriter SafeWriter(MakeUnique<FCheckingMemoryWriter>(TSStorage, &NumFlushes));
 			FString("Hello").SerializeAsANSICharArray(SafeWriter);
 
 			AutoRTFM::Testing::Abort([&]
 			{
 				FString("World").SerializeAsANSICharArray(SafeWriter);
+				SafeWriter.Flush();
 				AutoRTFM::AbortTransaction();
 			});
 		});
 
 		REQUIRE(TSStorage == NormalStorage);
+		REQUIRE(NumFlushes == 0);
 	}
 
-	SECTION("Ctor, Commit(Abort(Write), Write)")
+	SECTION("Ctor, Commit(Abort(Write), Write, Flush)")
 	{
 		TArray<uint8> TSStorage;
-		FTransactionallySafeArchiveWriter SafeWriter(MakeUnique<FCheckingMemoryWriter>(TSStorage));
+		int NumFlushes = 0;
+		FTransactionallySafeArchiveWriter SafeWriter(MakeUnique<FCheckingMemoryWriter>(TSStorage, &NumFlushes));
 
 		AutoRTFM::Testing::Commit([&]
 		{
@@ -168,18 +291,21 @@ TEST_CASE("TransactionallySafeArchiveWriter.CommitAndAbortWork")
 			});
 
 			FString("Hello").SerializeAsANSICharArray(SafeWriter);
+			SafeWriter.Flush();
 		});
 
 		REQUIRE(TSStorage == NormalStorage);
+		REQUIRE(NumFlushes == 1);
 	}
 
-	SECTION("Commit(Ctor, Abort(Write), Write)")
+	SECTION("Commit(Ctor, Abort(Write), Write, Flush)")
 	{
 		TArray<uint8> TSStorage;
+		int NumFlushes = 0;
 
 		AutoRTFM::Testing::Commit([&]
 		{
-			FTransactionallySafeArchiveWriter SafeWriter(MakeUnique<FCheckingMemoryWriter>(TSStorage));
+			FTransactionallySafeArchiveWriter SafeWriter(MakeUnique<FCheckingMemoryWriter>(TSStorage, &NumFlushes));
 
 			AutoRTFM::Testing::Abort([&]
 			{
@@ -188,22 +314,26 @@ TEST_CASE("TransactionallySafeArchiveWriter.CommitAndAbortWork")
 			});
 
 			FString("Hello").SerializeAsANSICharArray(SafeWriter);
+			SafeWriter.Flush();
 		});
 
 		REQUIRE(TSStorage == NormalStorage);
+		REQUIRE(NumFlushes == 1);
 	}
 
-	SECTION("Commit(Ctor, Abort(Write), Write)")
+	SECTION("Commit(Ctor, Abort(Write, Flush), Write)")
 	{
 		TArray<uint8> TSStorage;
+		int NumFlushes = 0;
 
 		AutoRTFM::Testing::Commit([&]
 		{
-			FTransactionallySafeArchiveWriter SafeWriter(MakeUnique<FCheckingMemoryWriter>(TSStorage));
+			FTransactionallySafeArchiveWriter SafeWriter(MakeUnique<FCheckingMemoryWriter>(TSStorage, &NumFlushes));
 
 			AutoRTFM::Testing::Abort([&]
 			{
 				FString("World").SerializeAsANSICharArray(SafeWriter);
+				SafeWriter.Flush();
 				AutoRTFM::AbortTransaction();
 			});
 
@@ -211,12 +341,14 @@ TEST_CASE("TransactionallySafeArchiveWriter.CommitAndAbortWork")
 		});
 
 		REQUIRE(TSStorage == NormalStorage);
+		REQUIRE(NumFlushes == 0);
 	}
 
-	SECTION("Ctor, Write, Abort(Commit(Write)))")
+	SECTION("Ctor, Write, Abort(Commit(Write, Flush)))")
 	{
 		TArray<uint8> TSStorage;
-		FTransactionallySafeArchiveWriter SafeWriter(MakeUnique<FCheckingMemoryWriter>(TSStorage));
+		int NumFlushes = 0;
+		FTransactionallySafeArchiveWriter SafeWriter(MakeUnique<FCheckingMemoryWriter>(TSStorage, &NumFlushes));
 
 		FString("Hello").SerializeAsANSICharArray(SafeWriter);
 
@@ -225,11 +357,13 @@ TEST_CASE("TransactionallySafeArchiveWriter.CommitAndAbortWork")
 			AutoRTFM::Testing::Commit([&]
 			{
 				FString("World").SerializeAsANSICharArray(SafeWriter);
+				SafeWriter.Flush();
 			});
 			AutoRTFM::AbortTransaction();
 		});
 
 		REQUIRE(TSStorage == NormalStorage);
+		REQUIRE(NumFlushes == 0);
 	}
 }
 
@@ -246,7 +380,8 @@ TEST_CASE("TransactionallySafeArchiveWriter.MixingWrites")
 	SECTION("Wrap, Write, Release, Write")
 	{
 		TArray<uint8> TSStorage;
-		FTransactionallySafeArchiveWriter SafeWriter(MakeUnique<FCheckingMemoryWriter>(TSStorage));
+		int NumFlushes = 0;
+		FTransactionallySafeArchiveWriter SafeWriter(MakeUnique<FCheckingMemoryWriter>(TSStorage, &NumFlushes));
 		FString("Hello").SerializeAsANSICharArray(SafeWriter);
 
 		TUniquePtr<FArchive> ReleasedWriter = SafeWriter.Release();
@@ -258,7 +393,8 @@ TEST_CASE("TransactionallySafeArchiveWriter.MixingWrites")
 	SECTION("Wrap, Commit(Write), Release, Write")
 	{
 		TArray<uint8> TSStorage;
-		FTransactionallySafeArchiveWriter SafeWriter(MakeUnique<FCheckingMemoryWriter>(TSStorage));
+		int NumFlushes = 0;
+		FTransactionallySafeArchiveWriter SafeWriter(MakeUnique<FCheckingMemoryWriter>(TSStorage, &NumFlushes));
 
 		AutoRTFM::Testing::Commit([&]
 		{
@@ -274,7 +410,8 @@ TEST_CASE("TransactionallySafeArchiveWriter.MixingWrites")
 	SECTION("Write, Wrap, Commit(Write), Release")
 	{
 		TArray<uint8> TSStorage;
-		TUniquePtr<FCheckingMemoryWriter> CheckingWriter(MakeUnique<FCheckingMemoryWriter>(TSStorage));
+		int NumFlushes = 0;
+		TUniquePtr<FCheckingMemoryWriter> CheckingWriter(MakeUnique<FCheckingMemoryWriter>(TSStorage, &NumFlushes));
 
 		FString("Hello").SerializeAsANSICharArray(*CheckingWriter);
 
@@ -295,7 +432,8 @@ TEST_CASE("TransactionallySafeArchiveWriter.MixingWrites")
 	SECTION("Write, Wrap, Release, Write")
 	{
 		TArray<uint8> TSStorage;
-		TUniquePtr<FCheckingMemoryWriter> CheckingWriter(MakeUnique<FCheckingMemoryWriter>(TSStorage));
+		int NumFlushes = 0;
+		TUniquePtr<FCheckingMemoryWriter> CheckingWriter(MakeUnique<FCheckingMemoryWriter>(TSStorage, &NumFlushes));
 
 		FString("Hello").SerializeAsANSICharArray(*CheckingWriter);
 
@@ -319,7 +457,8 @@ TEST_CASE("TransactionallySafeArchiveWriter.OnCommitWorks")
 	SECTION("Commit(OnCommit)")
 	{
 		TArray<uint8> TSStorage;
-		FTransactionallySafeArchiveWriter SafeWriter(MakeUnique<FCheckingMemoryWriter>(TSStorage));
+		int NumFlushes = 0;
+		FTransactionallySafeArchiveWriter SafeWriter(MakeUnique<FCheckingMemoryWriter>(TSStorage, &NumFlushes));
 
 		AutoRTFM::Testing::Commit([&]
 		{
@@ -328,17 +467,19 @@ TEST_CASE("TransactionallySafeArchiveWriter.OnCommitWorks")
 				FString("One").SerializeAsANSICharArray(SafeWriter);
 				FString("Two").SerializeAsANSICharArray(SafeWriter);
 				FString("Three").SerializeAsANSICharArray(SafeWriter);
+				SafeWriter.Flush();
 			});
 		});
 
 		REQUIRE(TSStorage == NormalStorage);
+		REQUIRE(NumFlushes == 1);
 	}
 
-
 	SECTION("Write, Commit(Write, OnCommit)")
 	{
 		TArray<uint8> TSStorage;
-		FTransactionallySafeArchiveWriter SafeWriter(MakeUnique<FCheckingMemoryWriter>(TSStorage));
+		int NumFlushes = 0;
+		FTransactionallySafeArchiveWriter SafeWriter(MakeUnique<FCheckingMemoryWriter>(TSStorage, &NumFlushes));
 
 		FString("One").SerializeAsANSICharArray(SafeWriter);
 
@@ -358,7 +499,8 @@ TEST_CASE("TransactionallySafeArchiveWriter.OnCommitWorks")
 	SECTION("Commit(OnCommit, Write), Write")
 	{
 		TArray<uint8> TSStorage;
-		FTransactionallySafeArchiveWriter SafeWriter(MakeUnique<FCheckingMemoryWriter>(TSStorage));
+		int NumFlushes = 0;
+		FTransactionallySafeArchiveWriter SafeWriter(MakeUnique<FCheckingMemoryWriter>(TSStorage, &NumFlushes));
 
 		AutoRTFM::Testing::Commit([&]
 		{
@@ -375,10 +517,11 @@ TEST_CASE("TransactionallySafeArchiveWriter.OnCommitWorks")
 		REQUIRE(TSStorage == NormalStorage);
 	}
 
-	SECTION("Commit(OnCommit, Write, OnCommit)")
+	SECTION("Commit(OnCommit, Write, OnCommit, OnCommit)")
 	{
 		TArray<uint8> TSStorage;
-		FTransactionallySafeArchiveWriter SafeWriter(MakeUnique<FCheckingMemoryWriter>(TSStorage));
+		int NumFlushes = 0;
+		FTransactionallySafeArchiveWriter SafeWriter(MakeUnique<FCheckingMemoryWriter>(TSStorage, &NumFlushes));
 
 		AutoRTFM::Testing::Commit([&]
 		{
@@ -393,8 +536,14 @@ TEST_CASE("TransactionallySafeArchiveWriter.OnCommitWorks")
 			{
 				FString("Three").SerializeAsANSICharArray(SafeWriter);
 			});
+
+			AutoRTFM::OnCommit([&]
+			{
+				SafeWriter.Flush();
+			});
 		});
 
 		REQUIRE(TSStorage == NormalStorage);
+		REQUIRE(NumFlushes == 1);
 	}
 }
diff --git a/Engine/Source/Runtime/Core/Public/Serialization/TransactionallySafeArchiveWriter.h b/Engine/Source/Runtime/Core/Public/Serialization/TransactionallySafeArchiveWriter.h
--- a/Engine/Source/Runtime/Core/Public/Serialization/TransactionallySafeArchiveWriter.h
+++ b/Engine/Source/Runtime/Core/Public/Serialization/TransactionallySafeArchiveWriter.h
@@ -23,6 +23,8 @@
  */
 class FTransactionallySafeArchiveWriter : public FArchive
 {
+	using DeferredWriteBuffer = TArray<uint8>;
+
 public:
 	FTransactionallySafeArchiveWriter(TUniquePtr<FArchive> Ar) : InnerArchive(MoveTemp(Ar))
 	{
@@ -37,9 +39,9 @@ public:
 		{
 			// Our wrapper is going away, so move everything into a lambda-capture for safekeeping.
 			AutoRTFM::PopOnCommitHandler(this);
-			AutoRTFM::PushOnCommitHandler(this, [Archive = InnerArchive.Release(), DataToSerialize = MoveTemp(DeferredWrites)]() mutable
+			AutoRTFM::PushOnCommitHandler(this, [Archive = InnerArchive.Release(), WriteBuffer = MoveTemp(DeferredWrites), bFlush = bRequestFlush]() mutable
 			{
-				Archive->Serialize(DataToSerialize.GetData(), DataToSerialize.NumBytes() );
+				DoDeferredWrites(*Archive, WriteBuffer, bFlush);
 				delete Archive;
 			});
 		}
@@ -71,20 +73,28 @@ public:
 		return 0;
 	}
 
+	/** Like all writes, flushes also need to be deferred to commit time. */
+	virtual void Flush() override
+	{
+		if (MaybeRegisterCommitHandler())
+		{
+			// Defer the flush until commit time.
+			bRequestFlush = true;
+		}
+		else
+		{
+			// We aren't in a transaction and don't have a commit handler, so nothing needs to be done. Forward to the inner archive.
+			InnerArchive->Flush();
+		}
+	}
+
 	virtual void Serialize(void* Data, int64 Num) override
 	{
 		check(InnerArchive);
 
-		if (bRegisteredCommitHandler)
-		{
-			// If we already have a registered commit handler, we are either within a transaction or are committing/aborting.
-			check(AutoRTFM::IsTransactional() || AutoRTFM::IsCommittingOrAborting());
-			DeferredWrites.Append(static_cast<uint8*>(Data), IntCastChecked<DeferredWriteBuffer::SizeType>(Num));
-		}
-		else if (AutoRTFM::IsTransactional())
+		if (MaybeRegisterCommitHandler())
 		{
-			// If we don't have a commit handler set up yet, but find ourselves running transactionally, set up the commit handler now.
-			RegisterCommitHandler();
+			// Defer serialization until commit time.
 			DeferredWrites.Append(static_cast<uint8*>(Data), IntCastChecked<DeferredWriteBuffer::SizeType>(Num));
 		}
 		else
@@ -96,6 +106,25 @@ public:
 	}
 
 private:
+	bool MaybeRegisterCommitHandler()
+	{
+		// Returns true if a commit handler is in use.
+		if (bRegisteredCommitHandler)
+		{
+			// If we already have a registered commit handler, we are either within a transaction or are committing/aborting.
+			check(AutoRTFM::IsTransactional() || AutoRTFM::IsCommittingOrAborting());
+			return true;
+		}
+		else if (AutoRTFM::IsTransactional())
+		{
+			// If we don't have a commit handler set up yet, but find ourselves running transactionally, set up the commit handler now.
+			RegisterCommitHandler();
+			return true;
+		}
+		// We don't need a commit handler and can forward on requests directly to the inner archive.
+		return false;
+	}
+
 	void RegisterCommitHandler()
 	{
 		check(!bRegisteredCommitHandler);
@@ -103,15 +132,25 @@ private:
 
 		AutoRTFM::PushOnCommitHandler(this, [this]
 		{
-			InnerArchive->Serialize(DeferredWrites.GetData(), DeferredWrites.NumBytes());
-			DeferredWrites.Reset();
+			DoDeferredWrites(*InnerArchive, DeferredWrites, bRequestFlush);
 			bRegisteredCommitHandler = false;
 		});
 	}
 
-	using DeferredWriteBuffer = TArray<uint8>;
+	static void DoDeferredWrites(FArchive& Archive, DeferredWriteBuffer& WriteBuffer, bool& bRequestFlush)
+	{
+		Archive.Serialize(WriteBuffer.GetData(), WriteBuffer.NumBytes());
+		WriteBuffer.Reset();
+
+		if (bRequestFlush)
+		{
+			Archive.Flush();
+			bRequestFlush = false;
+		}
+	}
 
 	TUniquePtr<FArchive> InnerArchive;
 	DeferredWriteBuffer DeferredWrites;
 	bool bRegisteredCommitHandler = false;
+	bool bRequestFlush = false;
 };
\ No newline at end of file
