[AutoRTFM] Make UGCObjectReferencer::[Add|Remove]Object() transaction-safe

Add abort handlers to both these methods to remove / add the object on abortion.

Remove the abort handling logic from FGCObject::RegisterGCObject(). 'this' may be destructed by the time the abort handler is called, which can cause corruption.

diff --git a/Engine/Source/Programs/AutoRTFMTests/Private/Catch2Includes.h b/Engine/Source/Programs/AutoRTFMTests/Private/Catch2Includes.h
--- a/Engine/Source/Programs/AutoRTFMTests/Private/Catch2Includes.h
+++ b/Engine/Source/Programs/AutoRTFMTests/Private/Catch2Includes.h
@@ -59,14 +59,12 @@ THIRD_PARTY_INCLUDES_END
     } while(false)
 
 // Same as Catch2's INTERNAL_CATCH_SECTION, but the unique variable name is
-// passed in, and fails the test if called from a closed transaction.
-#define AUTORTFM_CATCH2_SECTION(VAR_NAME, ...)                              \
-        CATCH_INTERNAL_START_WARNINGS_SUPPRESSION                           \
-        CATCH_INTERNAL_SUPPRESS_UNUSED_VARIABLE_WARNINGS                    \
-        if ( Catch::Section const& VAR_NAME =                               \
-                Catch::Section( CATCH_INTERNAL_LINEINFO, __VA_ARGS__ );     \
-                AutoRTFMCatch2::EnsureNotClosed("SECTION") && VAR_NAME)     \
-        CATCH_INTERNAL_STOP_WARNINGS_SUPPRESSION
+// passed in, and can be used in a closed transaction.
+#define AUTORTFM_CATCH2_SECTION(VAR_NAME, ...)                                           \
+    CATCH_INTERNAL_START_WARNINGS_SUPPRESSION                                            \
+    CATCH_INTERNAL_SUPPRESS_UNUSED_VARIABLE_WARNINGS                                     \
+    if (auto VAR_NAME = AutoRTFMCatch2::Section( CATCH_INTERNAL_LINEINFO, __VA_ARGS__ )) \
+    CATCH_INTERNAL_STOP_WARNINGS_SUPPRESSION
 
 #define REQUIRE( ... ) AUTORTFM_CATCH2_INTERNAL_CATCH_TEST( "REQUIRE", Catch::ResultDisposition::Normal, __VA_ARGS__ )
 #define REQUIRE_THAT( arg, matcher ) AUTORTFM_CATCH2_INTERNAL_CHECK_THAT( "REQUIRE_THAT", matcher, Catch::ResultDisposition::Normal, arg )
@@ -80,14 +78,23 @@ THIRD_PARTY_INCLUDES_END
 namespace AutoRTFMCatch2
 {
 
-inline bool EnsureNotClosed(const char* Where)
-{
-     if (AutoRTFM::IsClosed())
-     {
-        FAIL(Where << " cannot be used from a closed transaction");
-        return false;
-     }
-     return true;
-}
+// An AutoRTFM always-open wrapper around Catch::Section
+class Section {
+public:
+    UE_AUTORTFM_ALWAYS_OPEN
+    Section(Catch::SectionInfo&& Info) : Inner(std::move(Info)) {}
+
+    UE_AUTORTFM_ALWAYS_OPEN
+    Section(Catch::SourceLineInfo const& LineInfo, Catch::StringRef Name, const char* const = nullptr) : Inner(LineInfo, Name) {}
+
+    UE_AUTORTFM_ALWAYS_OPEN
+    ~Section() = default;
+
+    UE_AUTORTFM_ALWAYS_OPEN
+    explicit operator bool() const { return Inner.operator bool(); }
+
+private:
+    Catch::Section Inner;
+};
 
 } // AutoRTFMCatch2
diff --git a/Engine/Source/Programs/AutoRTFMTests/Private/UObjectTests.cpp b/Engine/Source/Programs/AutoRTFMTests/Private/UObjectTests.cpp
--- a/Engine/Source/Programs/AutoRTFMTests/Private/UObjectTests.cpp
+++ b/Engine/Source/Programs/AutoRTFMTests/Private/UObjectTests.cpp
@@ -867,3 +867,416 @@ TEST_CASE("UObject.ReleaseRef")
 		});
 	}
 }
+
+TEST_CASE("TObjectPtr")
+{
+	UMyAutoRTFMTestObject* Object = NewObject<UMyAutoRTFMTestObject>();
+
+	SECTION("Construct")
+	{
+		SECTION("Commit")
+		{
+			AutoRTFM::Testing::Commit([&]
+			{
+				TObjectPtr<UMyAutoRTFMTestObject> ObjectPtr(Object);
+			});
+		}
+
+		SECTION("Abort")
+		{
+			AutoRTFM::Testing::Abort([&]
+			{
+				TObjectPtr<UMyAutoRTFMTestObject> ObjectPtr(Object);
+				AutoRTFM::AbortTransaction();
+			});
+		}
+	}
+
+	SECTION("Destruct")
+	{
+		std::optional<TObjectPtr<UMyAutoRTFMTestObject>> ObjectPtr = TObjectPtr<UMyAutoRTFMTestObject>(Object);
+		SECTION("Commit")
+		{
+			AutoRTFM::Testing::Commit([&]
+			{
+				ObjectPtr.reset();
+			});
+		}
+
+		SECTION("Abort")
+		{
+			AutoRTFM::Testing::Abort([&]
+			{
+				ObjectPtr.reset();
+				AutoRTFM::AbortTransaction();
+			});
+		}
+	}
+}
+
+TEST_CASE("FGCObject")
+{
+	struct FMyGCObject : FGCObject
+	{
+		bool bAlive = true;
+
+		~FMyGCObject()
+		{
+			REQUIRE(bAlive);
+			bAlive = false;
+		}
+
+		void AddReferencedObjects([[maybe_unused]] FReferenceCollector& Collector) override
+		{
+			REQUIRE(bAlive);
+		}
+
+		FString GetReferencerName() const override
+		{
+			REQUIRE(bAlive);
+			return {};
+		}
+	};
+
+	auto Test = [&](auto&& Callback)
+	{
+		SECTION("NonTransactional")
+		{
+			Callback();
+		}
+
+		SECTION("Commit")
+		{
+			AutoRTFM::Testing::Commit([&]
+			{
+				Callback();
+			});
+		}
+
+		SECTION("Abort")
+		{
+			AutoRTFM::Testing::Abort([&]
+			{
+				Callback();
+				AutoRTFM::AbortTransaction();
+			});
+		}
+	};
+
+	// Tests for FGCObject that is allocated on the heap.
+	SECTION("Heap")
+	{
+		SECTION("Transact(Construct, Destruct)")
+		{
+			Test([&]
+			{
+				TUniquePtr<FMyGCObject> Object = MakeUnique<FMyGCObject>();
+			});
+		}
+
+		SECTION("Transact(Construct, Unregister, Register, Destruct)")
+		{
+			Test([&]
+			{
+				TUniquePtr<FMyGCObject> Object = MakeUnique<FMyGCObject>();
+				Object->UnregisterGCObject();
+				Object->RegisterGCObject();
+			});
+		}
+
+		SECTION("Transact(Construct, Unregister, Destruct)")
+		{
+			Test([&]
+			{
+				TUniquePtr<FMyGCObject> Object = MakeUnique<FMyGCObject>();
+				Object->UnregisterGCObject();
+			});
+		}
+
+		SECTION("Transact(Construct, Register, Destruct)")
+		{
+			Test([&]
+			{
+				TUniquePtr<FMyGCObject> Object = MakeUnique<FMyGCObject>();
+				Object->RegisterGCObject();
+			});
+		}
+
+		SECTION("Transact(Construct, Register, Unregister, Destruct)")
+		{
+			Test([&]
+			{
+				TUniquePtr<FMyGCObject> Object = MakeUnique<FMyGCObject>();
+				Object->RegisterGCObject();
+				Object->UnregisterGCObject();
+			});
+		}
+
+		SECTION("Transact(Construct, Unregister, Unregister, Register, Register, Destruct)")
+		{
+			Test([&]
+			{
+				TUniquePtr<FMyGCObject> Object = MakeUnique<FMyGCObject>();
+				Object->UnregisterGCObject();
+				Object->UnregisterGCObject();
+				Object->RegisterGCObject();
+				Object->RegisterGCObject();
+			});
+		}
+
+		SECTION("Transact(Construct), Destruct")
+		{
+			TUniquePtr<FMyGCObject> Object;
+			Test([&]
+			{
+				Object = MakeUnique<FMyGCObject>();
+			});
+		}
+
+		SECTION("Construct, Transact(Destruct)")
+		{
+			TUniquePtr<FMyGCObject> Object = MakeUnique<FMyGCObject>();
+			Test([&]
+			{
+				Object.Reset();
+			});
+		}
+
+		SECTION("Construct, Transact(Unregister), Destruct")
+		{
+			TUniquePtr<FMyGCObject> Object = MakeUnique<FMyGCObject>();
+			Test([&]
+			{
+				Object->UnregisterGCObject();
+			});
+		}
+
+		SECTION("Construct, Transact(Unregister, Register), Destruct")
+		{
+			TUniquePtr<FMyGCObject> Object = MakeUnique<FMyGCObject>();
+			Test([&]
+			{
+				Object->UnregisterGCObject();
+				Object->RegisterGCObject();
+			});
+		}
+
+		SECTION("Construct, Transact(Unregister, Destruct)")
+		{
+			TUniquePtr<FMyGCObject> Object = MakeUnique<FMyGCObject>();
+			Test([&]
+			{
+				Object->UnregisterGCObject();
+				Object.Reset();
+			});
+		}
+
+		SECTION("Construct, Transact(Unregister, Register, Destruct)")
+		{
+			TUniquePtr<FMyGCObject> Object = MakeUnique<FMyGCObject>();
+			Test([&]
+			{
+				Object->UnregisterGCObject();
+				Object->RegisterGCObject();
+				Object.Reset();
+			});
+		}
+
+		SECTION("Construct, Transact(Unregister, Transact(Register)), Destruct")
+		{
+			TUniquePtr<FMyGCObject> Object = MakeUnique<FMyGCObject>();
+			Test([&]
+			{
+				Object->UnregisterGCObject();
+				Test([&]
+				{
+					Object->RegisterGCObject();
+				});
+			});
+		}
+
+		SECTION("Construct, Transact(Unregister, Register, Transact(Unregister)), Destruct")
+		{
+			TUniquePtr<FMyGCObject> Object = MakeUnique<FMyGCObject>();
+			Test([&]
+			{
+				Object->UnregisterGCObject();
+				Object->RegisterGCObject();
+				Test([&]
+				{
+					Object->UnregisterGCObject();
+				});
+			});
+		}
+
+		SECTION("Construct, Transact(Unregister, Transact(Register), Destruct)")
+		{
+			TUniquePtr<FMyGCObject> Object = MakeUnique<FMyGCObject>();
+			Test([&]
+			{
+				Object->UnregisterGCObject();
+				Test([&]
+				{
+					Object->RegisterGCObject();
+				});
+				Object.Reset();
+			});
+		}
+
+		SECTION("Construct, Transact(Unregister, Register, Transact(Unregister), Destruct)")
+		{
+			TUniquePtr<FMyGCObject> Object = MakeUnique<FMyGCObject>();
+			Test([&]
+			{
+				Object->UnregisterGCObject();
+				Object->RegisterGCObject();
+				Test([&]
+				{
+					Object->UnregisterGCObject();
+				});
+				Object.Reset();
+			});
+		}
+	}
+
+	// Tests for FGCObject that is allocated within the transaction's stack.
+	SECTION("Stack")
+	{
+		SECTION("Transact(Construct, Destruct)")
+		{
+			Test([&]
+			{
+				FMyGCObject MyGCObject;
+			});
+		}
+
+		SECTION("Transact(Construct, Unregister, Destruct)")
+		{
+			Test([&]
+			{
+				FMyGCObject MyGCObject;
+				MyGCObject.UnregisterGCObject();
+			});
+		}
+
+		SECTION("Transact(Construct, Unregister, Unregister, Destruct)")
+		{
+			Test([&]
+			{
+				FMyGCObject MyGCObject;
+				MyGCObject.UnregisterGCObject();
+				MyGCObject.UnregisterGCObject();
+				MyGCObject.RegisterGCObject();
+			});
+		}
+
+		SECTION("Transact(Construct, Unregister, Register, Destruct)")
+		{
+			Test([&]
+			{
+				FMyGCObject MyGCObject;
+				MyGCObject.UnregisterGCObject();
+				MyGCObject.RegisterGCObject();
+			});
+		}
+
+		SECTION("Transact(Construct, Register, Register, Destruct)")
+		{
+			Test([&]
+			{
+				FMyGCObject MyGCObject;
+				MyGCObject.RegisterGCObject();
+				MyGCObject.RegisterGCObject();
+			});
+		}
+
+		SECTION("Transact(Construct, Register, Unregister, Register, Destruct)")
+		{
+			Test([&]
+			{
+				FMyGCObject MyGCObject;
+				MyGCObject.RegisterGCObject();
+				MyGCObject.UnregisterGCObject();
+				MyGCObject.RegisterGCObject();
+			});
+		}
+
+		SECTION("Transact(Construct, Transact(Unregister, Register), Destruct)")
+		{
+			Test([&]
+			{
+				FMyGCObject MyGCObject;
+				Test([&]
+				{
+					MyGCObject.UnregisterGCObject();
+					MyGCObject.RegisterGCObject();
+				});
+			});
+		}
+
+		SECTION("Transact(Construct, Unregister, Transact(Register), Destruct)")
+		{
+			Test([&]
+			{
+				FMyGCObject MyGCObject;
+				MyGCObject.UnregisterGCObject();
+				Test([&]
+				{
+					MyGCObject.RegisterGCObject();
+				});
+			});
+		}
+
+		SECTION("Transact(Construct, Transact(Unregister), Register, Destruct)")
+		{
+			Test([&]
+			{
+				FMyGCObject MyGCObject;
+				Test([&]
+				{
+					MyGCObject.UnregisterGCObject();
+				});
+				MyGCObject.RegisterGCObject();
+			});
+		}
+
+		SECTION("Construct, Transact(Destruct)")
+		{
+			std::optional<FMyGCObject> MyGCObject{FMyGCObject{}};
+			Test([&]
+			{
+				MyGCObject.reset();
+			});
+		}
+
+		SECTION("Transact(Construct), Destruct")
+		{
+			std::optional<FMyGCObject> MyGCObject;
+			Test([&]
+			{
+				MyGCObject = FMyGCObject{};
+			});
+		}
+
+		SECTION("Construct, Transact(Unregister, Destruct)")
+		{
+			std::optional<FMyGCObject> MyGCObject{FMyGCObject{}};
+			Test([&]
+			{
+				MyGCObject->UnregisterGCObject();
+				MyGCObject.reset();
+			});
+		}
+
+		SECTION("Transact(Construct, Unregister), Destruct")
+		{
+			std::optional<FMyGCObject> MyGCObject;
+			Test([&]
+			{
+				MyGCObject = FMyGCObject{};
+				MyGCObject->UnregisterGCObject();
+			});
+		}
+	}
+}
+
diff --git a/Engine/Source/Runtime/AutoRTFM/Private/API.cpp b/Engine/Source/Runtime/AutoRTFM/Private/API.cpp
--- a/Engine/Source/Runtime/AutoRTFM/Private/API.cpp
+++ b/Engine/Source/Runtime/AutoRTFM/Private/API.cpp
@@ -993,18 +993,18 @@ extern "C" UE_AUTORTFM_NOAUTORTFM void autortfm_unregister_open_to_closed_functi
 
 extern "C" bool autortfm_is_on_current_transaction_stack(void* Ptr)
 {
-	if (!FContext::Get()->IsTransactional())
-	{
-		return false;
-	}
-	else if (FTransaction* CurrentTransaction = FContext::Get()->GetCurrentTransaction())
-	{
-		return CurrentTransaction->IsOnStack(Ptr);
-	}
-	else
+	if (FContext* Context = FContext::Get())
 	{
-		return false;
+		if (!Context->IsTransactional())
+		{
+			return false;
+		}
+		if (FTransaction* CurrentTransaction = Context->GetCurrentTransaction())
+		{
+			return CurrentTransaction->IsOnStack(Ptr);
+		}
 	}
+	return false;
 }
 
 void ForTheRuntime::Initialize(const ForTheRuntime::FExternAPI& ExternAPI)
diff --git a/Engine/Source/Runtime/AutoRTFM/Private/StdLibSupport.cpp b/Engine/Source/Runtime/AutoRTFM/Private/StdLibSupport.cpp
--- a/Engine/Source/Runtime/AutoRTFM/Private/StdLibSupport.cpp
+++ b/Engine/Source/Runtime/AutoRTFM/Private/StdLibSupport.cpp
@@ -653,6 +653,8 @@ UE_AUTORTFM_REGISTER_OPEN_TO_CLOSED_FUNCTIONS(
 	UE_AUTORTFM_MAP_OPEN_TO_SELF(GetLocalTime),
 	UE_AUTORTFM_MAP_OPEN_TO_SELF(GetFileAttributesW),
 	UE_AUTORTFM_MAP_OPEN_TO_CLOSED(TlsSetValue, RTFM_TlsSetValue),
+	UE_AUTORTFM_MAP_OPEN_TO_SELF(_Query_perf_frequency),
+	UE_AUTORTFM_MAP_OPEN_TO_SELF(_Query_perf_counter),
 #else
 	UE_AUTORTFM_MAP_OPEN_TO_SELF(_ZNSt3__112__next_primeEm),
 #endif
diff --git a/Engine/Source/Runtime/CoreUObject/Private/Misc/GCObjectReferencer.cpp b/Engine/Source/Runtime/CoreUObject/Private/Misc/GCObjectReferencer.cpp
--- a/Engine/Source/Runtime/CoreUObject/Private/Misc/GCObjectReferencer.cpp
+++ b/Engine/Source/Runtime/CoreUObject/Private/Misc/GCObjectReferencer.cpp
@@ -4,6 +4,7 @@
 	GCObjectReferencer.cpp: Implementation of UGCObjectReferencer
 =============================================================================*/
 
+#include "AutoRTFM/AutoRTFM.h"
 #include "CoreMinimal.h"
 #include "Misc/ScopeLock.h"
 #include "UObject/ObjectMacros.h"
@@ -80,12 +81,24 @@ void UGCObjectReferencer::AddObject(FGCObject* Object)
 {
 	check(Object);
 	check(GObjUnhashUnreachableIsInProgress || GObjIncrementalPurgeIsInProgress || !IsGarbageCollectingAndLockingUObjectHashTables());
-	
-	FScopeLock ReferencedObjectsLock(&Impl->ReferencedObjectsCritical);
-	TArray<FGCObject*>& ReferencedObjects = Object->bCanMakeInitialReferences ? Impl->InitialReferencedObjects : Impl->RemainingReferencedObjects; 
-	// Make sure there are no duplicates. Should be impossible...
-	checkSlow(!ReferencedObjects.Contains(Object));
-	ReferencedObjects.Add(Object);
+
+	TArray<FGCObject*>* ReferencedObjects = Object->bCanMakeInitialReferences ? &Impl->InitialReferencedObjects : &Impl->RemainingReferencedObjects;
+
+	UE_AUTORTFM_OPEN
+	{
+		FScopeLock ReferencedObjectsLock(&Impl->ReferencedObjectsCritical);
+		// Make sure there are no duplicates. Should be impossible...
+		checkSlow(!ReferencedObjects->Contains(Object));
+		ReferencedObjects->Add(Object);
+	};
+
+	// Remove the object on abort.
+	UE_AUTORTFM_ONABORT(this, Object, ReferencedObjects)
+	{
+		FScopeLock ReferencedObjectsLock(&Impl->ReferencedObjectsCritical);
+		int32 NumRemoved = ReferencedObjects->RemoveSingleSwap(Object);
+		check(NumRemoved == 1);
+	};
 
 #if WITH_EDITORONLY_DATA
 	OnGCObjectAdded.Broadcast(Object);
@@ -97,10 +110,21 @@ void UGCObjectReferencer::RemoveObject(FGCObject* Object)
 	check(Object);
 	check(GObjUnhashUnreachableIsInProgress || GObjIncrementalPurgeIsInProgress || !IsGarbageCollectingAndLockingUObjectHashTables());
 	
-	FScopeLock ReferencedObjectsLock(&Impl->ReferencedObjectsCritical);
-	TArray<FGCObject*>& ReferencedObjects = Object->bCanMakeInitialReferences ? Impl->InitialReferencedObjects : Impl->RemainingReferencedObjects; 
-	int32 NumRemoved = ReferencedObjects.RemoveSingleSwap(Object);
-	check(NumRemoved == 1);
+	TArray<FGCObject*>* ReferencedObjects = Object->bCanMakeInitialReferences ? &Impl->InitialReferencedObjects : &Impl->RemainingReferencedObjects;
+
+	UE_AUTORTFM_OPEN
+	{
+		FScopeLock ReferencedObjectsLock(&Impl->ReferencedObjectsCritical);
+		int32 NumRemoved = ReferencedObjects->RemoveSingleSwap(Object);
+		check(NumRemoved == 1);
+	};
+
+	// Re-add the object on abort.
+	UE_AUTORTFM_ONABORT(this, Object, ReferencedObjects)
+	{
+		FScopeLock ReferencedObjectsLock(&Impl->ReferencedObjectsCritical);
+		ReferencedObjects->Add(Object);
+	};
 }
 
 bool UGCObjectReferencer::GetReferencerName(UObject* Object, FString& OutName, bool bOnlyIfAddingReferenced) const
@@ -255,6 +279,7 @@ UGCObjectReferencer* FGCObject::GGCObjectReferencer = nullptr;
 // works around a Linux/.so undefined typeinfo link error for plugins using -frtti
 const TCHAR* FGCObject::UnknownGCObjectName = UE::GC::UnknownGCObjectName.GetData();
 
+UE_AUTORTFM_ALWAYS_OPEN
 void FGCObject::StaticInit()
 {
 	if (GGCObjectReferencer == nullptr)
@@ -271,17 +296,7 @@ void FGCObject::RegisterGCObject()
 	{
 		StaticInit();
 
-		// Add this instance to the referencer's list
-		UE_AUTORTFM_OPEN
-		{
-			GGCObjectReferencer->AddObject(this);
-		};
-
-		// But if we abort, we don't want to leave a dangling reference!
-		UE_AUTORTFM_ONABORT(this)
-		{
-			GGCObjectReferencer->RemoveObject(this);
-		};
+		GGCObjectReferencer->AddObject(this);
 
 		bReferenceAdded = true;
 	}
@@ -294,6 +309,7 @@ void FGCObject::UnregisterGCObject()
 	if (GGCObjectReferencer && bReferenceAdded)
 	{
 		GGCObjectReferencer->RemoveObject(this);
+
 		bReferenceAdded = false;
 	}
 }
