[AutoRTFM] Move mutex tests out from UECoreTests into a separate CPP file.

Previously, these were all a single test case with dozens of sections. I've split it up so that each outermost section is now its own test case. This also helped uncover some redundant code, as one test case was repeated verbatim in two places.

diff --git a/Engine/Source/Programs/AutoRTFMTests/Private/MutexTests.cpp b/Engine/Source/Programs/AutoRTFMTests/Private/MutexTests.cpp
new file mode 100644
--- /dev/null
+++ b/Engine/Source/Programs/AutoRTFMTests/Private/MutexTests.cpp
@@ -0,0 +1,836 @@
+// Copyright Epic Games, Inc. All Rights Reserved.
+
+#include "AutoRTFM.h"
+#include "AutoRTFMTesting.h"
+#include "AutoRTFMTestUtils.h"
+#include "Catch2Includes.h"
+#include "Async/TransactionallySafeMutex.h"
+
+#include <atomic>
+#include <condition_variable>
+#include <cstring>
+#include <memory>
+#include <mutex>
+#include <thread>
+
+TEST_CASE("TransactionallySafeMutex.Outside Transaction")
+{
+	UE::FTransactionallySafeMutex Mutex;
+
+	AutoRTFM::Testing::Abort([&]
+	{
+		UE::TScopeLock Lock(Mutex);
+		AutoRTFM::AbortTransaction();
+	});
+
+	AutoRTFM::Testing::Commit([&]
+	{
+		UE::TScopeLock Lock(Mutex);
+	});
+}
+
+TEST_CASE("TransactionallySafeMutex.IsLocked Outside Transaction")
+{
+	UE::FTransactionallySafeMutex Mutex;
+
+	REQUIRE(!Mutex.IsLocked());
+	Mutex.Lock();
+	REQUIRE(Mutex.IsLocked());
+	Mutex.Unlock();
+	REQUIRE(!Mutex.IsLocked());
+}
+
+TEST_CASE("TransactionallySafeMutex.Inside Transaction")
+{
+	AutoRTFM::Testing::Abort([&]
+	{
+		UE::FTransactionallySafeMutex Mutex;
+		UE::TScopeLock Lock(Mutex);
+		AutoRTFM::AbortTransaction();
+	});
+
+	AutoRTFM::Testing::Commit([&]
+	{
+		UE::FTransactionallySafeMutex Mutex;
+		UE::TScopeLock Lock(Mutex);
+	});
+}
+
+TEST_CASE("TransactionallySafeMutex.IsLocked Inside Transaction")
+{
+	UE::FTransactionallySafeMutex Mutex;
+	REQUIRE(!Mutex.IsLocked());
+
+	AutoRTFM::Testing::Commit([&]
+	{
+		REQUIRE(!Mutex.IsLocked());
+		Mutex.Lock();
+		REQUIRE(Mutex.IsLocked());
+		Mutex.Unlock();
+		REQUIRE(Mutex.IsLocked());
+	});
+
+	REQUIRE(!Mutex.IsLocked());
+}
+
+TEST_CASE("TransactionallySafeMutex.Inside Transaction Used In Nested Transaction")
+{
+	AutoRTFM::Testing::Abort([&]
+	{
+		UE::FTransactionallySafeMutex Mutex;
+
+		AutoRTFM::Testing::Abort([&]
+		{
+			UE::TScopeLock Lock(Mutex);
+			AutoRTFM::CascadingAbortTransaction();
+		});
+	});
+
+	AutoRTFM::Testing::Commit([&]
+	{
+		UE::FTransactionallySafeMutex Mutex;
+
+		AutoRTFM::Testing::Abort([&]
+		{
+			UE::TScopeLock Lock(Mutex);
+			AutoRTFM::AbortTransaction();
+		});
+	});
+
+	AutoRTFM::Testing::Abort([&]
+	{
+		UE::FTransactionallySafeMutex Mutex;
+
+		AutoRTFM::Testing::Commit([&]
+		{
+			UE::TScopeLock Lock(Mutex);
+		});
+
+		AutoRTFM::AbortTransaction();
+	});
+
+	AutoRTFM::Testing::Commit([&]
+	{
+		UE::FTransactionallySafeMutex Mutex;
+
+		AutoRTFM::Testing::Commit([&]
+		{
+			UE::TScopeLock Lock(Mutex);
+		});
+	});
+}
+
+TEST_CASE("TransactionallySafeMutex.In Static Local Initializer")
+{
+	struct MyStruct final
+	{
+		UE::FTransactionallySafeMutex Mutex;
+	};
+
+	auto Lambda = []()
+	{
+		static MyStruct Mine;
+		UE::TScopeLock _(Mine.Mutex);
+		return 42;
+	};
+
+	AutoRTFM::Testing::Abort([&]
+	{
+		REQUIRE(42 == Lambda());
+		AutoRTFM::AbortTransaction();
+	});
+
+	REQUIRE(42 == Lambda());
+
+	AutoRTFM::Testing::Commit([&]
+	{
+		REQUIRE(42 == Lambda());
+	});
+
+	REQUIRE(42 == Lambda());
+}
+
+TEST_CASE("TransactionallySafeMutex.In Static Local Initializer Called From Open")
+{
+	struct MyStruct final
+	{
+		UE::FTransactionallySafeMutex Mutex;
+	};
+
+	auto Lambda = []()
+	{
+		static MyStruct Mine;
+		UE::TScopeLock _(Mine.Mutex);
+		return 42;
+	};
+
+	AutoRTFM::Testing::Abort([&]
+	{
+		AutoRTFM::Open([&] { REQUIRE(42 == Lambda()); });
+		AutoRTFM::AbortTransaction();
+	});
+
+	REQUIRE(42 == Lambda());
+
+	AutoRTFM::Testing::Commit([&]
+	{
+		AutoRTFM::Open([&] { REQUIRE(42 == Lambda()); });
+	});
+
+	REQUIRE(42 == Lambda());
+}
+
+TEST_CASE("TransactionallySafeMutex.Delete Heap Allocated Mutex")
+{
+	SECTION("SingleThread") // Mutex owned and destructed by this thread
+	{
+		UE::FTransactionallySafeMutex* const Mutex = new UE::FTransactionallySafeMutex();
+
+		AutoRTFM::Testing::Abort([&]
+		{
+			Mutex->Lock();
+			delete Mutex;
+			AutoRTFM::AbortTransaction();
+		});
+
+		REQUIRE(!Mutex->IsLocked());
+
+		AutoRTFM::Testing::Commit([&]
+		{
+			Mutex->Lock();
+			delete Mutex;
+		});
+	}
+		
+	SECTION("MultiThread") // Mutex owned and destructed by another thread
+	{
+		// This test does not support retries due to coordination with another thread.
+		AutoRTFMTestUtils::FScopedRetry NoRetry(AutoRTFM::ForTheRuntime::EAutoRTFMRetryTransactionState::NoRetry);
+
+		enum EState
+		{
+			InitializeThread,
+			MutexReady,
+			MutexUsed,
+			MutexDeleted,
+		};
+
+		struct FEvent
+		{
+			UE_AUTORTFM_ALWAYS_OPEN
+			void Signal(EState State)
+			{
+				{
+					std::unique_lock Lock(Mutex);
+					CurrentState = State;
+				}
+				CV.notify_one();
+			}
+
+			UE_AUTORTFM_ALWAYS_OPEN
+			void Wait(EState State)
+			{
+				std::unique_lock Lock(Mutex);
+				CV.wait(Lock, [this, State] { return CurrentState == State; });
+			}
+		private:
+			EState CurrentState = InitializeThread;
+			std::mutex Mutex;
+			std::condition_variable CV;
+		};
+			
+		UE::FTransactionallySafeMutex* Mutex = nullptr;
+		FEvent Event;
+
+		std::thread Thread([&Mutex, &Event]
+		{
+			Mutex = new UE::FTransactionallySafeMutex();
+			Event.Signal(MutexReady);
+			Event.Wait(MutexUsed);
+			delete Mutex;
+			Event.Signal(MutexDeleted);
+		});
+
+		Event.Wait(MutexReady);
+
+		SECTION("Lock, Abort, Destroy")
+		{
+			AutoRTFM::Testing::Abort([&]
+			{
+				Mutex->Lock();
+				AutoRTFM::AbortTransaction();
+			});
+
+			REQUIRE(!Mutex->IsLocked());
+
+			Event.Signal(MutexUsed);
+			Event.Wait(MutexDeleted);
+		}
+
+		SECTION("Lock, Unlock, Destroy, Abort")
+		{
+			AutoRTFM::Testing::Abort([&]
+			{
+				Mutex->Lock();
+				Mutex->Unlock();
+
+				Event.Signal(MutexUsed);
+				Event.Wait(MutexDeleted);
+
+				AutoRTFM::AbortTransaction();
+			});
+		}
+
+		SECTION("Lock, Unlock, Destroy, Commit")
+		{
+			AutoRTFM::Testing::Commit([&]
+			{
+				Mutex->Lock();
+				Mutex->Unlock();
+
+				Event.Signal(MutexUsed);
+				Event.Wait(MutexDeleted);
+			});
+		}
+
+		Thread.join();
+	}
+}
+
+TEST_CASE("TransactionallySafeMutex.Lock Within, Unlock Outside")
+{
+	UE::FTransactionallySafeMutex Mutex;
+
+	AutoRTFM::Testing::Commit([&]
+	{
+		Mutex.Lock();
+	});
+
+	REQUIRE(Mutex.IsLocked());
+	Mutex.Unlock();
+	REQUIRE(!Mutex.IsLocked());
+}
+
+TEST_CASE("TransactionallySafeMutex.TryLock Within, Unlock Outside")
+{
+	UE::FTransactionallySafeMutex Mutex;
+
+	AutoRTFM::Testing::Commit([&]
+	{
+		REQUIRE(Mutex.TryLock());
+	});
+
+	REQUIRE(Mutex.IsLocked());
+	Mutex.Unlock();
+	REQUIRE(!Mutex.IsLocked());
+}
+
+TEST_CASE("TransactionallySafeMutex.Lock Outside, Unlock Within")
+{
+	UE::FTransactionallySafeMutex Mutex = UE::FTransactionallySafeMutex(UE::FAcquireLock());
+
+	AutoRTFM::Testing::Commit([&]
+	{
+		REQUIRE(Mutex.IsLocked());
+		Mutex.Unlock();
+	});
+
+	REQUIRE(!Mutex.IsLocked());
+}
+
+TEST_CASE("TransactionallySafeMutex.Lock, Unlock, Lock")
+{
+	UE::FTransactionallySafeMutex Mutex;
+
+	SECTION("Commit")
+	{
+		AutoRTFM::Testing::Commit([&]
+		{
+			Mutex.Lock();
+			Mutex.Unlock();
+			Mutex.Lock();
+		});
+
+		REQUIRE(Mutex.IsLocked());
+		Mutex.Unlock();
+	}
+
+	SECTION("Abort")
+	{
+		AutoRTFM::Testing::Abort([&]
+		{
+			Mutex.Lock();
+			Mutex.Unlock();
+			Mutex.Lock();
+			AutoRTFM::AbortTransaction();
+		});
+
+		REQUIRE(!Mutex.IsLocked());
+	}
+}
+
+TEST_CASE("TransactionallySafeMutex.TryLock, Unlock, TryLock")
+{
+	UE::FTransactionallySafeMutex Mutex;
+
+	SECTION("Commit")
+	{
+		AutoRTFM::Testing::Commit([&]
+		{
+			REQUIRE(Mutex.TryLock());
+			Mutex.Unlock();
+			REQUIRE(Mutex.TryLock());
+		});
+
+		REQUIRE(Mutex.IsLocked());
+		Mutex.Unlock();
+	}
+		
+	SECTION("Abort")
+	{
+		AutoRTFM::Testing::Abort([&]
+		{
+			REQUIRE(Mutex.TryLock());
+			Mutex.Unlock();
+			REQUIRE(Mutex.TryLock());
+			AutoRTFM::AbortTransaction();
+		});
+
+		REQUIRE(!Mutex.IsLocked());
+	}
+}
+
+TEST_CASE("TransactionallySafeMutex.Unlock, Lock, Unlock")
+{
+	UE::FTransactionallySafeMutex Mutex;
+	Mutex.Lock();
+
+	SECTION("Commit")
+	{
+		AutoRTFM::Testing::Commit([&]
+		{
+			Mutex.Unlock();
+			Mutex.Lock();
+			Mutex.Unlock();
+		});
+
+		REQUIRE(!Mutex.IsLocked());
+	}
+
+	SECTION("Commit")
+	{
+		AutoRTFM::Testing::Abort([&]
+		{
+			Mutex.Unlock();
+			Mutex.Lock();
+			Mutex.Unlock();
+			AutoRTFM::AbortTransaction();
+		});
+
+		REQUIRE(Mutex.IsLocked());
+		Mutex.Unlock();
+	}
+}
+
+TEST_CASE("TransactionallySafeMutex.Unlock, TryLock, Unlock")
+{
+	UE::FTransactionallySafeMutex Mutex;
+	Mutex.Lock();
+
+	SECTION("Commit")
+	{
+		AutoRTFM::Testing::Commit([&]
+		{
+			Mutex.Unlock();
+			REQUIRE(Mutex.TryLock());
+			Mutex.Unlock();
+		});
+
+		REQUIRE(!Mutex.IsLocked());
+	}
+
+	SECTION("Commit")
+	{
+		AutoRTFM::Testing::Abort([&]
+		{
+			Mutex.Unlock();
+			REQUIRE(Mutex.TryLock());
+			Mutex.Unlock();
+			AutoRTFM::AbortTransaction();
+		});
+
+		REQUIRE(Mutex.IsLocked());
+		Mutex.Unlock();
+	}
+}
+
+TEST_CASE("TransactionallySafeMutex.Lock Outside, Unlock & Lock Within")
+{
+	UE::FTransactionallySafeMutex Mutex = UE::FTransactionallySafeMutex(UE::FAcquireLock());
+
+	AutoRTFM::Testing::Commit([&]
+	{
+		REQUIRE(Mutex.IsLocked());
+		Mutex.Unlock();
+		Mutex.Lock();
+	});
+
+	REQUIRE(Mutex.IsLocked());
+}
+
+TEST_CASE("TransactionallySafeMutex.Commit(Lock, Commit(Unlock))")
+{
+	UE::FTransactionallySafeMutex Mutex;
+
+	AutoRTFM::Testing::Commit([&]
+	{
+		Mutex.Lock();
+
+		AutoRTFM::Testing::Commit([&]
+		{
+			Mutex.Unlock();
+		});
+
+		REQUIRE(Mutex.IsLocked());
+	});
+
+	REQUIRE(!Mutex.IsLocked());
+}
+
+TEST_CASE("TransactionallySafeMutex.Abort(Lock, Commit(Unlock, Lock))")
+{
+	UE::FTransactionallySafeMutex Mutex;
+
+	AutoRTFM::Testing::Abort([&]
+	{
+		Mutex.Lock();
+
+		AutoRTFM::Testing::Commit([&]
+		{
+			Mutex.Unlock();
+			Mutex.Lock();
+		});
+
+		REQUIRE(Mutex.IsLocked());
+
+		AutoRTFM::AbortTransaction();
+	});
+
+	REQUIRE(!Mutex.IsLocked());
+}
+
+TEST_CASE("TransactionallySafeMutex.Contended Lock")
+{
+	struct RAIIeroo final
+	{
+		RAIIeroo()
+		{
+			State = AutoRTFM::ForTheRuntime::GetRetryTransaction();
+			AutoRTFM::ForTheRuntime::SetRetryTransaction(AutoRTFM::ForTheRuntime::EAutoRTFMRetryTransactionState::NoRetry);
+		}
+
+		~RAIIeroo()
+		{
+			AutoRTFM::ForTheRuntime::SetRetryTransaction(State);
+		}
+
+	private:
+		AutoRTFM::ForTheRuntime::EAutoRTFMRetryTransactionState State;
+	} _;
+
+	SECTION("Contender Parks Transaction")
+	{
+		UE::FTransactionallySafeMutex Mutex;
+		std::atomic_uint Orderer = 0;
+
+		std::thread Contender([&]
+		{
+			REQUIRE(0 == Orderer);
+			Mutex.Lock();
+			Orderer += 1; // unblock main thread
+			while (1 == Orderer) {} // wait on main thread
+			std::this_thread::sleep_for(std::chrono::milliseconds(100));
+			Mutex.Unlock();
+		});
+
+		while (0 == Orderer) {} // wait on contender
+
+		AutoRTFM::Testing::Commit([&]
+		{
+			REQUIRE(Mutex.IsLocked());
+			AutoRTFM::Open([&] { Orderer += 1; }); // unblock contender
+			Mutex.Lock();
+		});
+
+		REQUIRE(Mutex.IsLocked());
+		Mutex.Unlock();
+		Contender.join();
+	}
+
+	SECTION("Transaction Parks Contender")
+	{
+		UE::FTransactionallySafeMutex Mutex;
+		std::atomic_uint Orderer = 0;
+
+		std::thread Contender([&]
+		{
+			while (0 == Orderer) {} // wait on main thread
+			Mutex.Lock();
+		});
+
+		AutoRTFM::Testing::Commit([&]
+		{
+			Mutex.Lock();
+			AutoRTFM::Open([&]
+			{
+				Orderer += 1; // unblock contender
+				std::this_thread::sleep_for(std::chrono::milliseconds(100));
+			});
+			Mutex.Unlock();
+		});
+
+		Contender.join();
+		Mutex.Unlock();
+	}
+}
+
+TEST_CASE("TransactionallySafeMutex.In On-Commit")
+{
+	UE::FTransactionallySafeMutex Mutex;
+
+	AutoRTFM::Testing::Commit([&]
+	{
+		Mutex.Lock();
+
+		AutoRTFM::OnCommit([&]
+		{
+			REQUIRE(Mutex.IsLocked());
+			Mutex.Unlock();
+			REQUIRE(!Mutex.IsLocked());
+		});
+	});
+
+	AutoRTFM::Testing::Commit([&]
+	{
+		AutoRTFM::OnCommit([&]
+		{
+			REQUIRE(Mutex.IsLocked());
+			Mutex.Unlock();
+			REQUIRE(Mutex.IsLocked());
+		});
+			
+		Mutex.Lock();
+	});
+
+	AutoRTFM::Testing::Commit([&]
+	{
+		REQUIRE(Mutex.TryLock());
+		REQUIRE(!Mutex.TryLock());
+
+		AutoRTFM::OnCommit([&]
+		{
+			REQUIRE(Mutex.IsLocked());
+			Mutex.Unlock();
+			REQUIRE(!Mutex.IsLocked());
+		});
+	});
+
+	AutoRTFM::Testing::Commit([&]
+	{
+		AutoRTFM::OnCommit([&]
+		{
+			REQUIRE(Mutex.IsLocked());
+			Mutex.Unlock();
+			REQUIRE(Mutex.IsLocked());
+		});
+			
+		REQUIRE(Mutex.TryLock());
+	});
+
+	Mutex.Lock();
+	REQUIRE(Mutex.IsLocked());
+
+	AutoRTFM::Testing::Commit([&]
+	{
+		AutoRTFM::OnCommit([&]
+		{
+			REQUIRE(Mutex.IsLocked());
+			Mutex.Lock();
+			REQUIRE(Mutex.IsLocked());
+		});
+			
+		Mutex.Unlock();
+	});
+
+	REQUIRE(Mutex.IsLocked());
+
+	AutoRTFM::Testing::Commit([&]
+	{
+		Mutex.Unlock();
+
+		AutoRTFM::OnCommit([&]
+		{
+			REQUIRE(!Mutex.IsLocked());
+			Mutex.Lock();
+			REQUIRE(Mutex.IsLocked());
+		});
+	});
+
+	REQUIRE(Mutex.IsLocked());
+		
+	AutoRTFM::Testing::Commit([&]
+	{
+		Mutex.Unlock();
+
+		AutoRTFM::OnCommit([&]
+		{
+			REQUIRE(Mutex.TryLock());
+		});
+	});
+
+	REQUIRE(Mutex.IsLocked());
+		
+	AutoRTFM::Testing::Commit([&]
+	{
+		AutoRTFM::OnCommit([&]
+		{
+			REQUIRE(Mutex.TryLock());
+		});
+
+		Mutex.Unlock();
+	});
+
+	REQUIRE(Mutex.IsLocked());
+}
+
+TEST_CASE("TransactionallySafeMutex.In On-Abort")
+{
+	UE::FTransactionallySafeMutex Mutex;
+
+	AutoRTFM::Testing::Abort([&]
+	{
+		AutoRTFM::OnAbort([&]
+		{
+			REQUIRE(!Mutex.IsLocked());
+		});
+
+		Mutex.Lock();
+
+		AutoRTFM::OnAbort([&]
+		{
+			REQUIRE(Mutex.IsLocked());
+		});
+
+		AutoRTFM::AbortTransaction();
+	});
+
+	AutoRTFM::Testing::Abort([&]
+	{
+		AutoRTFM::OnAbort([&]
+		{
+			REQUIRE(!Mutex.IsLocked());
+		});
+			
+		REQUIRE(Mutex.TryLock());
+
+		AutoRTFM::OnAbort([&]
+		{
+			REQUIRE(Mutex.IsLocked());
+		});
+
+		AutoRTFM::AbortTransaction();
+	});
+
+	AutoRTFM::Testing::Abort([&]
+	{
+		AutoRTFM::OnAbort([&]
+		{
+			REQUIRE(!Mutex.IsLocked());
+		});
+			
+		REQUIRE(Mutex.TryLock());
+
+		AutoRTFM::AbortTransaction();
+	});
+
+	Mutex.Lock();
+	REQUIRE(Mutex.IsLocked());
+
+	AutoRTFM::Testing::Abort([&]
+	{
+		AutoRTFM::OnAbort([&]
+		{
+			REQUIRE(Mutex.IsLocked());
+		});
+			
+		Mutex.Unlock();
+
+		AutoRTFM::AbortTransaction();
+	});
+
+	REQUIRE(Mutex.IsLocked());
+
+	AutoRTFM::Testing::Abort([&]
+	{
+		Mutex.Unlock();
+
+		AutoRTFM::OnAbort([&]
+		{
+			REQUIRE(Mutex.IsLocked());
+		});
+
+		AutoRTFM::AbortTransaction();
+	});
+
+	REQUIRE(Mutex.IsLocked());
+		
+	AutoRTFM::Testing::Abort([&]
+	{
+		Mutex.Unlock();
+
+		AutoRTFM::OnAbort([&]
+		{
+			REQUIRE(!Mutex.TryLock());
+		});
+
+		AutoRTFM::AbortTransaction();
+	});
+
+	REQUIRE(Mutex.IsLocked());
+		
+	AutoRTFM::Testing::Abort([&]
+	{
+		AutoRTFM::OnAbort([&]
+		{
+			REQUIRE(!Mutex.TryLock());
+		});
+
+		Mutex.Unlock();
+
+		AutoRTFM::AbortTransaction();
+	});
+
+	REQUIRE(Mutex.IsLocked());
+}
+
+TEST_CASE("TransactionallySafeMutex.Locked Mutex In Destructed Object")
+{
+	struct MyStruct final
+	{
+		UE::FTransactionallySafeMutex Mutex;
+
+		~MyStruct()
+		{
+			memset(this, 0, sizeof(MyStruct));
+		}
+	};
+
+	std::unique_ptr<MyStruct> Mine(new MyStruct);
+
+	AutoRTFM::Testing::Commit([&]
+	{
+		Mine->Mutex.Lock();
+		Mine.reset();
+		REQUIRE(!Mine);
+	});
+}
diff --git a/Engine/Source/Programs/AutoRTFMTests/Private/UECoreTests.cpp b/Engine/Source/Programs/AutoRTFMTests/Private/UECoreTests.cpp
--- a/Engine/Source/Programs/AutoRTFMTests/Private/UECoreTests.cpp
+++ b/Engine/Source/Programs/AutoRTFMTests/Private/UECoreTests.cpp
@@ -1130,844 +1130,6 @@ TEST_CASE("UECore.TScopeLock_TransactionallySafeCriticalSection")
 	}
 }
 
-TEST_CASE("UECore.TransactionallySafeMutex")
-{
-	SECTION("Outside Transaction")
-	{
-		UE::FTransactionallySafeMutex Mutex;
-
-		AutoRTFM::Testing::Abort([&]
-		{
-			UE::TScopeLock Lock(Mutex);
-			AutoRTFM::AbortTransaction();
-		});
-
-		AutoRTFM::Testing::Commit([&]
-		{
-			UE::TScopeLock Lock(Mutex);
-		});
-	}
-
-	SECTION("IsLocked Outside Transaction")
-	{
-		UE::FTransactionallySafeMutex Mutex;
-
-		REQUIRE(!Mutex.IsLocked());
-		Mutex.Lock();
-		REQUIRE(Mutex.IsLocked());
-		Mutex.Unlock();
-		REQUIRE(!Mutex.IsLocked());
-	}
-
-	SECTION("Inside Transaction")
-	{
-		AutoRTFM::Testing::Abort([&]
-		{
-			UE::FTransactionallySafeMutex Mutex;
-			UE::TScopeLock Lock(Mutex);
-			AutoRTFM::AbortTransaction();
-		});
-
-		AutoRTFM::Testing::Commit([&]
-		{
-			UE::FTransactionallySafeMutex Mutex;
-			UE::TScopeLock Lock(Mutex);
-		});
-	}
-
-	SECTION("IsLocked Inside Transaction")
-	{
-		UE::FTransactionallySafeMutex Mutex;
-		REQUIRE(!Mutex.IsLocked());
-
-		AutoRTFM::Testing::Commit([&]
-		{
-			REQUIRE(!Mutex.IsLocked());
-			Mutex.Lock();
-			REQUIRE(Mutex.IsLocked());
-			Mutex.Unlock();
-			REQUIRE(Mutex.IsLocked());
-		});
-
-		REQUIRE(!Mutex.IsLocked());
-	}
-
-	SECTION("Inside Transaction Used In Nested Transaction")
-	{
-		AutoRTFM::Testing::Abort([&]
-		{
-			UE::FTransactionallySafeMutex Mutex;
-
-			AutoRTFM::Testing::Abort([&]
-			{
-				UE::TScopeLock Lock(Mutex);
-				AutoRTFM::CascadingAbortTransaction();
-			});
-		});
-
-		AutoRTFM::Testing::Commit([&]
-		{
-			UE::FTransactionallySafeMutex Mutex;
-
-			AutoRTFM::Testing::Abort([&]
-			{
-				UE::TScopeLock Lock(Mutex);
-				AutoRTFM::AbortTransaction();
-			});
-		});
-
-		AutoRTFM::Testing::Abort([&]
-		{
-			UE::FTransactionallySafeMutex Mutex;
-
-			AutoRTFM::Testing::Commit([&]
-			{
-				UE::TScopeLock Lock(Mutex);
-			});
-
-			AutoRTFM::AbortTransaction();
-		});
-
-		AutoRTFM::Testing::Commit([&]
-		{
-			UE::FTransactionallySafeMutex Mutex;
-
-			AutoRTFM::Testing::Commit([&]
-			{
-				UE::TScopeLock Lock(Mutex);
-			});
-		});
-	}
-
-	SECTION("In Static Local Initializer")
-	{
-		struct MyStruct final
-		{
-			UE::FTransactionallySafeMutex Mutex;
-		};
-
-		auto Lambda = []()
-		{
-			static MyStruct Mine;
-			UE::TScopeLock _(Mine.Mutex);
-			return 42;
-		};
-
-		AutoRTFM::Testing::Abort([&]
-		{
-			REQUIRE(42 == Lambda());
-			AutoRTFM::AbortTransaction();
-		});
-
-		REQUIRE(42 == Lambda());
-
-		AutoRTFM::Testing::Commit([&]
-		{
-			REQUIRE(42 == Lambda());
-		});
-
-		REQUIRE(42 == Lambda());
-	}
-
-	SECTION("In Static Local Initializer Called From Open")
-	{
-		struct MyStruct final
-		{
-			UE::FTransactionallySafeMutex Mutex;
-		};
-
-		auto Lambda = []()
-		{
-			static MyStruct Mine;
-			UE::TScopeLock _(Mine.Mutex);
-			return 42;
-		};
-
-		AutoRTFM::Testing::Abort([&]
-		{
-			AutoRTFM::Open([&] { REQUIRE(42 == Lambda()); });
-			AutoRTFM::AbortTransaction();
-		});
-
-		REQUIRE(42 == Lambda());
-
-		AutoRTFM::Testing::Commit([&]
-		{
-			AutoRTFM::Open([&] { REQUIRE(42 == Lambda()); });
-		});
-
-		REQUIRE(42 == Lambda());
-	}
-
-	SECTION("Delete Heap Allocated Mutex")
-	{
-		SECTION("SingleThread") // Mutex owned and destructed by this thread
-		{
-			UE::FTransactionallySafeMutex* const Mutex = new UE::FTransactionallySafeMutex();
-
-			AutoRTFM::Testing::Abort([&]
-			{
-				Mutex->Lock();
-				delete Mutex;
-				AutoRTFM::AbortTransaction();
-			});
-
-			REQUIRE(!Mutex->IsLocked());
-
-			AutoRTFM::Testing::Commit([&]
-			{
-				Mutex->Lock();
-				delete Mutex;
-			});
-		}
-		
-		SECTION("MultiThread") // Mutex owned and destructed by another thread
-		{
-			// This test does not support retries due to coordination with another thread.
-			AutoRTFMTestUtils::FScopedRetry NoRetry(AutoRTFM::ForTheRuntime::EAutoRTFMRetryTransactionState::NoRetry);
-
-			enum EState
-			{
-				InitializeThread,
-				MutexReady,
-				MutexUsed,
-				MutexDeleted,
-			};
-
-			struct FEvent
-			{
-				UE_AUTORTFM_ALWAYS_OPEN
-				void Signal(EState State)
-				{
-					{
-						std::unique_lock Lock(Mutex);
-						CurrentState = State;
-					}
-					CV.notify_one();
-				}
-
-				UE_AUTORTFM_ALWAYS_OPEN
-				void Wait(EState State)
-				{
-					std::unique_lock Lock(Mutex);
-					CV.wait(Lock, [this, State] { return CurrentState == State; });
-				}
-			private:
-				EState CurrentState = InitializeThread;
-				std::mutex Mutex;
-				std::condition_variable CV;
-			};
-			
-			UE::FTransactionallySafeMutex* Mutex = nullptr;
-			FEvent Event;
-
-			std::thread Thread([&Mutex, &Event]
-			{
-				Mutex = new UE::FTransactionallySafeMutex();
-				Event.Signal(MutexReady);
-				Event.Wait(MutexUsed);
-				delete Mutex;
-				Event.Signal(MutexDeleted);
-			});
-
-			Event.Wait(MutexReady);
-
-			SECTION("Lock, Abort, Destroy")
-			{
-				AutoRTFM::Testing::Abort([&]
-				{
-					Mutex->Lock();
-					AutoRTFM::AbortTransaction();
-				});
-
-				REQUIRE(!Mutex->IsLocked());
-
-				Event.Signal(MutexUsed);
-				Event.Wait(MutexDeleted);
-			}
-
-			SECTION("Lock, Unlock, Destroy, Abort")
-			{
-				AutoRTFM::Testing::Abort([&]
-				{
-					Mutex->Lock();
-					Mutex->Unlock();
-
-					Event.Signal(MutexUsed);
-					Event.Wait(MutexDeleted);
-
-					AutoRTFM::AbortTransaction();
-				});
-			}
-
-			SECTION("Lock, Unlock, Destroy, Commit")
-			{
-				AutoRTFM::Testing::Commit([&]
-				{
-					Mutex->Lock();
-					Mutex->Unlock();
-
-					Event.Signal(MutexUsed);
-					Event.Wait(MutexDeleted);
-				});
-			}
-
-			Thread.join();
-		}
-	}
-
-	SECTION("Lock Within, Unlock Outside")
-	{
-		UE::FTransactionallySafeMutex Mutex;
-
-		AutoRTFM::Testing::Commit([&]
-		{
-			Mutex.Lock();
-		});
-
-		REQUIRE(Mutex.IsLocked());
-		Mutex.Unlock();
-		REQUIRE(!Mutex.IsLocked());
-	}
-
-	SECTION("TryLock Within, Unlock Outside")
-	{
-		UE::FTransactionallySafeMutex Mutex;
-
-		AutoRTFM::Testing::Commit([&]
-		{
-			REQUIRE(Mutex.TryLock());
-		});
-
-		REQUIRE(Mutex.IsLocked());
-		Mutex.Unlock();
-		REQUIRE(!Mutex.IsLocked());
-	}
-
-	SECTION("Lock Outside, Unlock Within")
-	{
-		UE::FTransactionallySafeMutex Mutex = UE::FTransactionallySafeMutex(UE::FAcquireLock());
-
-		AutoRTFM::Testing::Commit([&]
-		{
-			REQUIRE(Mutex.IsLocked());
-			Mutex.Unlock();
-		});
-
-		REQUIRE(!Mutex.IsLocked());
-	}
-
-	SECTION("Lock, Unlock, Lock")
-	{
-		UE::FTransactionallySafeMutex Mutex;
-
-		SECTION("Commit")
-		{
-			AutoRTFM::Testing::Commit([&]
-			{
-				Mutex.Lock();
-				Mutex.Unlock();
-				Mutex.Lock();
-			});
-
-			REQUIRE(Mutex.IsLocked());
-			Mutex.Unlock();
-		}
-
-		SECTION("Abort")
-		{
-			AutoRTFM::Testing::Abort([&]
-			{
-				Mutex.Lock();
-				Mutex.Unlock();
-				Mutex.Lock();
-				AutoRTFM::AbortTransaction();
-			});
-
-			REQUIRE(!Mutex.IsLocked());
-		}
-	}
-
-	SECTION("TryLock, Unlock, TryLock")
-	{
-		UE::FTransactionallySafeMutex Mutex;
-
-		SECTION("Commit")
-		{
-			AutoRTFM::Testing::Commit([&]
-			{
-				REQUIRE(Mutex.TryLock());
-				Mutex.Unlock();
-				REQUIRE(Mutex.TryLock());
-			});
-
-			REQUIRE(Mutex.IsLocked());
-			Mutex.Unlock();
-		}
-		
-		SECTION("Abort")
-		{
-			AutoRTFM::Testing::Abort([&]
-			{
-				REQUIRE(Mutex.TryLock());
-				Mutex.Unlock();
-				REQUIRE(Mutex.TryLock());
-				AutoRTFM::AbortTransaction();
-			});
-
-			REQUIRE(!Mutex.IsLocked());
-		}
-	}
-
-	SECTION("Unlock, Lock, Unlock")
-	{
-		UE::FTransactionallySafeMutex Mutex;
-		Mutex.Lock();
-
-		SECTION("Commit")
-		{
-			AutoRTFM::Testing::Commit([&]
-			{
-				Mutex.Unlock();
-				Mutex.Lock();
-				Mutex.Unlock();
-			});
-
-			REQUIRE(!Mutex.IsLocked());
-		}
-
-		SECTION("Commit")
-		{
-			AutoRTFM::Testing::Abort([&]
-			{
-				Mutex.Unlock();
-				Mutex.Lock();
-				Mutex.Unlock();
-				AutoRTFM::AbortTransaction();
-			});
-
-			REQUIRE(Mutex.IsLocked());
-			Mutex.Unlock();
-		}
-	}
-
-	SECTION("Unlock, TryLock, Unlock")
-	{
-		UE::FTransactionallySafeMutex Mutex;
-		Mutex.Lock();
-
-		SECTION("Commit")
-		{
-			AutoRTFM::Testing::Commit([&]
-			{
-				Mutex.Unlock();
-				REQUIRE(Mutex.TryLock());
-				Mutex.Unlock();
-			});
-
-			REQUIRE(!Mutex.IsLocked());
-		}
-
-		SECTION("Commit")
-		{
-			AutoRTFM::Testing::Abort([&]
-			{
-				Mutex.Unlock();
-				REQUIRE(Mutex.TryLock());
-				Mutex.Unlock();
-				AutoRTFM::AbortTransaction();
-			});
-
-			REQUIRE(Mutex.IsLocked());
-			Mutex.Unlock();
-		}
-	}
-
-	SECTION("Lock Outside, Unlock Within")
-	{
-		UE::FTransactionallySafeMutex Mutex = UE::FTransactionallySafeMutex(UE::FAcquireLock());
-
-		AutoRTFM::Testing::Commit([&]
-		{
-			REQUIRE(Mutex.IsLocked());
-			Mutex.Unlock();
-		});
-
-		REQUIRE(!Mutex.IsLocked());
-	}
-
-	SECTION("Lock Outside, Unlock & Lock Within")
-	{
-		UE::FTransactionallySafeMutex Mutex = UE::FTransactionallySafeMutex(UE::FAcquireLock());
-
-		AutoRTFM::Testing::Commit([&]
-		{
-			REQUIRE(Mutex.IsLocked());
-			Mutex.Unlock();
-			Mutex.Lock();
-		});
-
-		REQUIRE(Mutex.IsLocked());
-	}
-
-	SECTION("Commit(Lock, Commit(Unlock))")
-	{
-		UE::FTransactionallySafeMutex Mutex;
-
-		AutoRTFM::Testing::Commit([&]
-		{
-			Mutex.Lock();
-
-			AutoRTFM::Testing::Commit([&]
-			{
-				Mutex.Unlock();
-			});
-
-			REQUIRE(Mutex.IsLocked());
-		});
-
-		REQUIRE(!Mutex.IsLocked());
-	}
-
-	SECTION("Abort(Lock, Commit(Unlock, Lock))")
-	{
-		UE::FTransactionallySafeMutex Mutex;
-
-		AutoRTFM::Testing::Abort([&]
-		{
-			Mutex.Lock();
-
-			AutoRTFM::Testing::Commit([&]
-			{
-				Mutex.Unlock();
-				Mutex.Lock();
-			});
-
-			REQUIRE(Mutex.IsLocked());
-
-			AutoRTFM::AbortTransaction();
-		});
-
-		REQUIRE(!Mutex.IsLocked());
-	}
-
-	SECTION("Contended Lock")
-	{
-		struct RAIIeroo final
-		{
-			RAIIeroo()
-			{
-				State = AutoRTFM::ForTheRuntime::GetRetryTransaction();
-				AutoRTFM::ForTheRuntime::SetRetryTransaction(AutoRTFM::ForTheRuntime::EAutoRTFMRetryTransactionState::NoRetry);
-			}
-
-			~RAIIeroo()
-			{
-				AutoRTFM::ForTheRuntime::SetRetryTransaction(State);
-			}
-
-		private:
-			AutoRTFM::ForTheRuntime::EAutoRTFMRetryTransactionState State;
-		} _;
-
-		SECTION("Contender Parks Transaction")
-		{
-			UE::FTransactionallySafeMutex Mutex;
-			std::atomic_uint Orderer = 0;
-
-			std::thread Contender([&]
-			{
-				REQUIRE(0 == Orderer);
-				Mutex.Lock();
-				Orderer += 1; // unblock main thread
-				while (1 == Orderer) {} // wait on main thread
-				std::this_thread::sleep_for(std::chrono::milliseconds(100));
-				Mutex.Unlock();
-			});
-
-			while (0 == Orderer) {} // wait on contender
-
-			AutoRTFM::Testing::Commit([&]
-			{
-				REQUIRE(Mutex.IsLocked());
-				AutoRTFM::Open([&] { Orderer += 1; }); // unblock contender
-				Mutex.Lock();
-			});
-
-			REQUIRE(Mutex.IsLocked());
-			Mutex.Unlock();
-			Contender.join();
-		}
-
-		SECTION("Transaction Parks Contender")
-		{
-			UE::FTransactionallySafeMutex Mutex;
-			std::atomic_uint Orderer = 0;
-
-			std::thread Contender([&]
-			{
-				while (0 == Orderer) {} // wait on main thread
-				Mutex.Lock();
-			});
-
-			AutoRTFM::Testing::Commit([&]
-			{
-				Mutex.Lock();
-				AutoRTFM::Open([&]
-				{
-					Orderer += 1; // unblock contender
-					std::this_thread::sleep_for(std::chrono::milliseconds(100));
-				});
-				Mutex.Unlock();
-			});
-
-			Contender.join();
-			Mutex.Unlock();
-		}
-	}
-
-	SECTION("In On-Commit")
-	{
-		UE::FTransactionallySafeMutex Mutex;
-
-		AutoRTFM::Testing::Commit([&]
-		{
-			Mutex.Lock();
-
-			AutoRTFM::OnCommit([&]
-			{
-				REQUIRE(Mutex.IsLocked());
-				Mutex.Unlock();
-				REQUIRE(!Mutex.IsLocked());
-			});
-		});
-
-		AutoRTFM::Testing::Commit([&]
-		{
-			AutoRTFM::OnCommit([&]
-			{
-				REQUIRE(Mutex.IsLocked());
-				Mutex.Unlock();
-				REQUIRE(Mutex.IsLocked());
-			});
-			
-			Mutex.Lock();
-		});
-
-		AutoRTFM::Testing::Commit([&]
-		{
-			REQUIRE(Mutex.TryLock());
-			REQUIRE(!Mutex.TryLock());
-
-			AutoRTFM::OnCommit([&]
-			{
-				REQUIRE(Mutex.IsLocked());
-				Mutex.Unlock();
-				REQUIRE(!Mutex.IsLocked());
-			});
-		});
-
-		AutoRTFM::Testing::Commit([&]
-		{
-			AutoRTFM::OnCommit([&]
-			{
-				REQUIRE(Mutex.IsLocked());
-				Mutex.Unlock();
-				REQUIRE(Mutex.IsLocked());
-			});
-			
-			REQUIRE(Mutex.TryLock());
-		});
-
-		Mutex.Lock();
-		REQUIRE(Mutex.IsLocked());
-
-		AutoRTFM::Testing::Commit([&]
-		{
-			AutoRTFM::OnCommit([&]
-			{
-				REQUIRE(Mutex.IsLocked());
-				Mutex.Lock();
-				REQUIRE(Mutex.IsLocked());
-			});
-			
-			Mutex.Unlock();
-		});
-
-		REQUIRE(Mutex.IsLocked());
-
-		AutoRTFM::Testing::Commit([&]
-		{
-			Mutex.Unlock();
-
-			AutoRTFM::OnCommit([&]
-			{
-				REQUIRE(!Mutex.IsLocked());
-				Mutex.Lock();
-				REQUIRE(Mutex.IsLocked());
-			});
-		});
-
-		REQUIRE(Mutex.IsLocked());
-		
-		AutoRTFM::Testing::Commit([&]
-		{
-			Mutex.Unlock();
-
-			AutoRTFM::OnCommit([&]
-			{
-				REQUIRE(Mutex.TryLock());
-			});
-		});
-
-		REQUIRE(Mutex.IsLocked());
-		
-		AutoRTFM::Testing::Commit([&]
-		{
-			AutoRTFM::OnCommit([&]
-			{
-				REQUIRE(Mutex.TryLock());
-			});
-
-			Mutex.Unlock();
-		});
-
-		REQUIRE(Mutex.IsLocked());
-	}
-
-	SECTION("In On-Abort")
-	{
-		UE::FTransactionallySafeMutex Mutex;
-
-		AutoRTFM::Testing::Abort([&]
-		{
-			AutoRTFM::OnAbort([&]
-			{
-				REQUIRE(!Mutex.IsLocked());
-			});
-
-			Mutex.Lock();
-
-			AutoRTFM::OnAbort([&]
-			{
-				REQUIRE(Mutex.IsLocked());
-			});
-
-			AutoRTFM::AbortTransaction();
-		});
-
-		AutoRTFM::Testing::Abort([&]
-		{
-			AutoRTFM::OnAbort([&]
-			{
-				REQUIRE(!Mutex.IsLocked());
-			});
-			
-			REQUIRE(Mutex.TryLock());
-
-			AutoRTFM::OnAbort([&]
-			{
-				REQUIRE(Mutex.IsLocked());
-			});
-
-			AutoRTFM::AbortTransaction();
-		});
-
-		AutoRTFM::Testing::Abort([&]
-		{
-			AutoRTFM::OnAbort([&]
-			{
-				REQUIRE(!Mutex.IsLocked());
-			});
-			
-			REQUIRE(Mutex.TryLock());
-
-			AutoRTFM::AbortTransaction();
-		});
-
-		Mutex.Lock();
-		REQUIRE(Mutex.IsLocked());
-
-		AutoRTFM::Testing::Abort([&]
-		{
-			AutoRTFM::OnAbort([&]
-			{
-				REQUIRE(Mutex.IsLocked());
-			});
-			
-			Mutex.Unlock();
-
-			AutoRTFM::AbortTransaction();
-		});
-
-		REQUIRE(Mutex.IsLocked());
-
-		AutoRTFM::Testing::Abort([&]
-		{
-			Mutex.Unlock();
-
-			AutoRTFM::OnAbort([&]
-			{
-				REQUIRE(Mutex.IsLocked());
-			});
-
-			AutoRTFM::AbortTransaction();
-		});
-
-		REQUIRE(Mutex.IsLocked());
-		
-		AutoRTFM::Testing::Abort([&]
-		{
-			Mutex.Unlock();
-
-			AutoRTFM::OnAbort([&]
-			{
-				REQUIRE(!Mutex.TryLock());
-			});
-
-			AutoRTFM::AbortTransaction();
-		});
-
-		REQUIRE(Mutex.IsLocked());
-		
-		AutoRTFM::Testing::Abort([&]
-		{
-			AutoRTFM::OnAbort([&]
-			{
-				REQUIRE(!Mutex.TryLock());
-			});
-
-			Mutex.Unlock();
-
-			AutoRTFM::AbortTransaction();
-		});
-
-		REQUIRE(Mutex.IsLocked());
-	}
-
-	SECTION("Locked Mutex In Destructed Object")
-	{
-		struct MyStruct final
-		{
-			UE::FTransactionallySafeMutex Mutex;
-
-			~MyStruct()
-			{
-				memset(this, 0, sizeof(MyStruct));
-			}
-		};
-
-		std::unique_ptr<MyStruct> Mine(new MyStruct);
-
-		AutoRTFM::Testing::Commit([&]
-		{
-			Mine->Mutex.Lock();
-			Mine.reset();
-			REQUIRE(!Mine);
-		});
-	}
-}
-
 TEST_CASE("UECore.FTextFormatPatternDefinition")
 {
 	FTextFormatPatternDefinitionConstPtr Ptr;
