[AutoRTFM] Fix transactional restoring of reconstructed FDynamicallyTypedValues

Due to the explicit open logic in the FDynamicallyTypedValue constructor, destructing and then re-constructing a FDynamicallyTypedValue within a closed transaction could fail to properly revert back to the original value on transaction abort.

The lengthy comment next to the UE_AUTORTFM_OPEN block explains why we can't run InitializeToNull() in the closed, so instead, we record the original value in the destructor.

diff --git a/Engine/Source/Programs/AutoRTFMTests/Private/UECoreTests.cpp b/Engine/Source/Programs/AutoRTFMTests/Private/UECoreTests.cpp
--- a/Engine/Source/Programs/AutoRTFMTests/Private/UECoreTests.cpp
+++ b/Engine/Source/Programs/AutoRTFMTests/Private/UECoreTests.cpp
@@ -22,6 +22,7 @@
 #include "Modules/ModuleManager.h"
 #include "Serialization/CustomVersion.h"
 #include "UObject/CoreRedirects.h"
+#include "UObject/DynamicallyTypedValue.h"
 #include "UObject/LinkerInstancingContext.h"
 #include "UObject/NameTypes.h"
 #include "UObject/Stack.h"
@@ -2312,3 +2313,191 @@ TEST_CASE("FVirtualStackAllocator.PreventsOOMIsValidationSafe")
 	FVirtualStackAllocator Allocator(8*1024*1024, EVirtualStackAllocatorDecommitMode::AllOnDestruction);
 	RunVStackAllocTest(&Allocator, 100);
 }
+
+TEST_CASE("FDynamicallyTypedValue")
+{
+	static constexpr uint64 InitializedValue = 42;
+
+	struct FType : UE::FDynamicallyTypedValueType
+	{
+		FType() : FDynamicallyTypedValueType(/* InNumBytes */ 8, /* InMinAlignmentLogTwo */ 3, /* InContainsReferences */ EContainsReferences::DoesNot) {}
+		void MarkReachable(FReferenceCollector& Collector) override {}
+		void MarkValueReachable(void* Data, FReferenceCollector& Collector) const override {}
+		void InitializeValue(void* Data) const override
+		{
+			*reinterpret_cast<uint64*>(Data) = InitializedValue;
+		}
+		void InitializeValueFromCopy(void* DestData, const void* SourceData) const override
+		{
+			*reinterpret_cast<uint64*>(DestData) = *reinterpret_cast<const uint64*>(SourceData);
+		}
+		void DestroyValue(void* Data) const override {}
+		void SerializeValue(FStructuredArchive::FSlot Slot, void* Data, const void* DefaultData) const override {}
+		uint32 GetValueHash(const void* Data) const override { return 0; }
+		bool AreIdentical(const void* DataA, const void* DataB) const override { return false; }
+	} Type;
+
+	auto GetValue = [](UE::FDynamicallyTypedValue& Value)
+	{
+		return *reinterpret_cast<uint64*>(Value.GetDataPointer());
+	};
+
+	SECTION("Construct")
+	{
+		SECTION("Commit")
+		{
+			AutoRTFM::Testing::Commit([&]
+			{
+				UE::FDynamicallyTypedValue Value;
+				REQUIRE(&Value.GetType() == &UE::FDynamicallyTypedValue::NullType());
+				REQUIRE(GetValue(Value) == 0);
+			});
+		}
+
+		SECTION("Abort")
+		{
+			AutoRTFM::Testing::Abort([&]
+			{
+				UE::FDynamicallyTypedValue Value;
+				REQUIRE(&Value.GetType() == &UE::FDynamicallyTypedValue::NullType());
+				REQUIRE(GetValue(Value) == 0);
+				AutoRTFM::AbortTransaction();
+			});
+		}
+	}
+
+	SECTION("Copy Construct")
+	{
+		UE::FDynamicallyTypedValue Original;
+		Original.InitializeAsType(Type);
+		REQUIRE(&Original.GetType() == &Type);
+		REQUIRE(GetValue(Original) == InitializedValue);
+
+		SECTION("Commit")
+		{
+			AutoRTFM::Testing::Commit([&]
+			{
+				UE::FDynamicallyTypedValue Value(Original);
+				REQUIRE(&Value.GetType() == &Type);
+				REQUIRE(GetValue(Value) == InitializedValue);
+			});
+		}
+
+		SECTION("Abort")
+		{
+			AutoRTFM::Testing::Abort([&]
+			{
+				UE::FDynamicallyTypedValue Value(Original);
+				REQUIRE(&Value.GetType() == &Type);
+				REQUIRE(GetValue(Value) == InitializedValue);
+				AutoRTFM::AbortTransaction();
+			});
+		}
+
+		REQUIRE(&Original.GetType() == &Type);
+		REQUIRE(GetValue(Original) == InitializedValue);
+	}
+
+	SECTION("Move Construct")
+	{
+		UE::FDynamicallyTypedValue Original;
+		Original.InitializeAsType(Type);
+		REQUIRE(&Original.GetType() == &Type);
+		REQUIRE(GetValue(Original) == InitializedValue);
+
+		SECTION("Commit")
+		{
+			AutoRTFM::Testing::Commit([&]
+			{
+				UE::FDynamicallyTypedValue Value(std::move(Original));
+				REQUIRE(&Value.GetType() == &Type);
+				REQUIRE(GetValue(Value) == InitializedValue);
+				REQUIRE(&Original.GetType() == &UE::FDynamicallyTypedValue::NullType());
+				REQUIRE(GetValue(Original) == 0);
+			});
+			REQUIRE(&Original.GetType() == &UE::FDynamicallyTypedValue::NullType());
+			REQUIRE(GetValue(Original) == 0);
+		}
+
+		SECTION("Abort")
+		{
+			AutoRTFM::Testing::Abort([&]
+			{
+				UE::FDynamicallyTypedValue Value(std::move(Original));
+				REQUIRE(&Value.GetType() == &Type);
+				REQUIRE(GetValue(Value) == InitializedValue);
+				REQUIRE(&Original.GetType() == &UE::FDynamicallyTypedValue::NullType());
+				REQUIRE(GetValue(Original) == 0);
+				AutoRTFM::AbortTransaction();
+			});
+			REQUIRE(&Original.GetType() == &Type);
+			REQUIRE(GetValue(Original) == InitializedValue);
+		}
+	}
+
+	SECTION("InitializeAsType")
+	{
+		UE::FDynamicallyTypedValue Value;
+		REQUIRE(&Value.GetType() == &UE::FDynamicallyTypedValue::NullType());
+		REQUIRE(GetValue(Value) == 0);
+
+		SECTION("Commit")
+		{
+			AutoRTFM::Testing::Commit([&]
+			{
+				Value.InitializeAsType(Type);
+				REQUIRE(&Value.GetType() == &Type);
+				REQUIRE(GetValue(Value) == InitializedValue);
+			});
+			REQUIRE(Value.GetDataPointer() != nullptr);
+		}
+
+		SECTION("Abort")
+		{
+			AutoRTFM::Testing::Abort([&]
+			{
+				Value.InitializeAsType(Type);
+				REQUIRE(&Value.GetType() == &Type);
+				REQUIRE(GetValue(Value) == InitializedValue);
+				AutoRTFM::AbortTransaction();
+			});
+			REQUIRE(&Value.GetType() == &UE::FDynamicallyTypedValue::NullType());
+			REQUIRE(GetValue(Value) == 0);
+		}
+	}
+
+	SECTION("Reconstruct")
+	{
+		UE::FDynamicallyTypedValue Value;
+		Value.InitializeAsType(Type);
+		REQUIRE(&Value.GetType() == &Type);
+		REQUIRE(GetValue(Value) == InitializedValue);
+
+		SECTION("Commit")
+		{
+			AutoRTFM::Testing::Commit([&]
+			{
+				Value.~FDynamicallyTypedValue();
+				new (&Value) UE::FDynamicallyTypedValue();
+				REQUIRE(&Value.GetType() == &UE::FDynamicallyTypedValue::NullType());
+				REQUIRE(GetValue(Value) == 0);
+			});
+			REQUIRE(GetValue(Value) == 0);
+		}
+
+		SECTION("Abort")
+		{
+			AutoRTFM::Testing::Abort([&]
+			{
+				Value.~FDynamicallyTypedValue();
+				new (&Value) UE::FDynamicallyTypedValue();
+				REQUIRE(&Value.GetType() == &UE::FDynamicallyTypedValue::NullType());
+				REQUIRE(GetValue(Value) == 0);
+				AutoRTFM::AbortTransaction();
+			});
+			REQUIRE(&Value.GetType() == &Type);
+			REQUIRE(GetValue(Value) == InitializedValue);
+		}
+	}
+}
+
diff --git a/Engine/Source/Runtime/CoreUObject/Public/UObject/DynamicallyTypedValue.h b/Engine/Source/Runtime/CoreUObject/Public/UObject/DynamicallyTypedValue.h
--- a/Engine/Source/Runtime/CoreUObject/Public/UObject/DynamicallyTypedValue.h
+++ b/Engine/Source/Runtime/CoreUObject/Public/UObject/DynamicallyTypedValue.h
@@ -232,6 +232,18 @@ namespace UE
 				FMemory::Free(HeapData);
 				HeapData = nullptr;
 			}
+			else if (::AutoRTFM::IsClosed())
+			{
+				// Assign to InlineData if we're in a closed transaction.
+				// This is done to ensure that the value of InlineData is
+				// recorded in the transaction. This is important as some
+				// code paths destruct then re-construct the 
+				// FDynamicallyTypedValue, and the constructor calls
+				// InitializeToNull() in the open (without recording the initial
+				// values). This can result in the FDynamicallyTypedValue being
+				// nulled without a write record to restore the original value.
+				InlineData = 0;
+			}
 		}
 	};
 }
