[AutoRTFM] Clean up MTTSAccessDetector tests in preparation for adding the MRSW access detector.

This CL templatizes the test to allow swapping in the MRSW access detector, and adopts the newer `AutoRTFM::Testing::Commit`/`Abort` helper methods to reduce code bloat a little bit. The indentation was also pulled in slightly.

diff --git a/Engine/Source/Programs/AutoRTFMTests/Private/MTTSAccessDetector.cpp b/Engine/Source/Programs/AutoRTFMTests/Private/MTTSAccessDetector.cpp
--- a/Engine/Source/Programs/AutoRTFMTests/Private/MTTSAccessDetector.cpp
+++ b/Engine/Source/Programs/AutoRTFMTests/Private/MTTSAccessDetector.cpp
@@ -9,907 +9,848 @@
 
 #if ENABLE_MT_DETECTOR
 
-TEST_CASE("MTTransactionallySafeAccessDetector")
+TEMPLATE_TEST_CASE("MTTransactionallySafeAccessDetector", "[MTAccessDetector]", FRWTransactionallySafeAccessDetector)
 {
 	SECTION("Transact(UE_MT_SCOPED_READ_ACCESS, Scope(UE_MT_SCOPED_READ_ACCESS), Abort)")
 	{
-		FRWTransactionallySafeAccessDetector Detector;
-		AutoRTFM::ETransactionResult Result = AutoRTFM::Transact([&]
+		TestType Detector;
+		AutoRTFM::Testing::Abort([&]
+		{
+			UE_MT_SCOPED_READ_ACCESS(Detector);
 			{
 				UE_MT_SCOPED_READ_ACCESS(Detector);
-				{
-					UE_MT_SCOPED_READ_ACCESS(Detector);
-				}
-				AutoRTFM::AbortTransaction();
-			});
-		REQUIRE(AutoRTFM::ETransactionResult::AbortedByRequest == Result);
+			}
+			AutoRTFM::AbortTransaction();
+		});
 	}
 
 	SECTION("UE_MT_SCOPED_READ_ACCESS, Transact(UE_MT_SCOPED_READ_ACCESS)")
 	{
-		FRWTransactionallySafeAccessDetector Detector;
+		TestType Detector;
 		UE_MT_SCOPED_READ_ACCESS(Detector);
-		AutoRTFM::ETransactionResult Result = AutoRTFM::Transact([&]
-			{
-				UE_MT_SCOPED_READ_ACCESS(Detector);
-			});
-		REQUIRE(AutoRTFM::ETransactionResult::Committed == Result);
+		AutoRTFM::Testing::Commit([&]
+		{
+			UE_MT_SCOPED_READ_ACCESS(Detector);
+		});
 	}
 
 	SECTION("UE_MT_SCOPED_READ_ACCESS, Transact(UE_MT_SCOPED_READ_ACCESS, Abort)")
 	{
-		FRWTransactionallySafeAccessDetector Detector;
+		TestType Detector;
 		UE_MT_SCOPED_READ_ACCESS(Detector);
-		AutoRTFM::ETransactionResult Result = AutoRTFM::Transact([&]
-			{
-				UE_MT_SCOPED_READ_ACCESS(Detector);
-				AutoRTFM::AbortTransaction();
-			});
-		REQUIRE(AutoRTFM::ETransactionResult::AbortedByRequest == Result);
+		AutoRTFM::Testing::Abort([&]
+		{
+			UE_MT_SCOPED_READ_ACCESS(Detector);
+			AutoRTFM::AbortTransaction();
+		});
 	}
 
 	SECTION("Transact(UE_MT_SCOPED_WRITE_ACCESS)")
 	{
-		FRWTransactionallySafeAccessDetector Detector;
-		AutoRTFM::ETransactionResult Result = AutoRTFM::Transact([&]
-			{
-				UE_MT_SCOPED_WRITE_ACCESS(Detector);
-			});
-		REQUIRE(AutoRTFM::ETransactionResult::Committed == Result);
+		TestType Detector;
+		AutoRTFM::Testing::Commit([&]
+		{
+			UE_MT_SCOPED_WRITE_ACCESS(Detector);
+		});
 	}
 
 	SECTION("Transact(UE_MT_SCOPED_WRITE_ACCESS, Abort)")
 	{
-		FRWTransactionallySafeAccessDetector Detector;
-		AutoRTFM::ETransactionResult Result = AutoRTFM::Transact([&]
-			{
-				UE_MT_SCOPED_WRITE_ACCESS(Detector);
-				AutoRTFM::AbortTransaction();
-			});
-		REQUIRE(AutoRTFM::ETransactionResult::AbortedByRequest == Result);
+		TestType Detector;
+		AutoRTFM::Testing::Abort([&]
+		{
+			UE_MT_SCOPED_WRITE_ACCESS(Detector);
+			AutoRTFM::AbortTransaction();
+		});
 	}
 
 	SECTION("Transact(Scope(UE_MT_SCOPED_WRITE_ACCESS), Abort)")
 	{
-		FRWTransactionallySafeAccessDetector Detector;
-		AutoRTFM::ETransactionResult Result = AutoRTFM::Transact([&]
+		TestType Detector;
+		AutoRTFM::Testing::Abort([&]
+		{
 			{
-				{
-					UE_MT_SCOPED_WRITE_ACCESS(Detector);
-				}
-				AutoRTFM::AbortTransaction();
-			});
-		REQUIRE(AutoRTFM::ETransactionResult::AbortedByRequest == Result);
+				UE_MT_SCOPED_WRITE_ACCESS(Detector);
+			}
+			AutoRTFM::AbortTransaction();
+		});
 	}
 
 	SECTION("Transact(Scope(UE_MT_SCOPED_READ_ACCESS), Scope(UE_MT_SCOPED_READ_ACCESS))")
 	{
-		FRWTransactionallySafeAccessDetector Detector;
-		AutoRTFM::ETransactionResult Result = AutoRTFM::Transact([&]
+		TestType Detector;
+		AutoRTFM::Testing::Commit([&]
+		{
 			{
-				{
-					UE_MT_SCOPED_READ_ACCESS(Detector);
-				}
-				{
-					UE_MT_SCOPED_READ_ACCESS(Detector);
-				}
-			});
-		REQUIRE(AutoRTFM::ETransactionResult::Committed == Result);
+				UE_MT_SCOPED_READ_ACCESS(Detector);
+			}
+			{
+				UE_MT_SCOPED_READ_ACCESS(Detector);
+			}
+		});
 	}
 
 	SECTION("Transact(Scope(UE_MT_SCOPED_READ_ACCESS), Scope(UE_MT_SCOPED_WRITE_ACCESS))")
 	{
-		FRWTransactionallySafeAccessDetector Detector;
-		AutoRTFM::ETransactionResult Result = AutoRTFM::Transact([&]
+		TestType Detector;
+		AutoRTFM::Testing::Commit([&]
+		{
 			{
-				{
-					UE_MT_SCOPED_READ_ACCESS(Detector);
-				}
-				{
-					UE_MT_SCOPED_WRITE_ACCESS(Detector);
-				}
-			});
-		REQUIRE(AutoRTFM::ETransactionResult::Committed == Result);
+				UE_MT_SCOPED_READ_ACCESS(Detector);
+			}
+			{
+				UE_MT_SCOPED_WRITE_ACCESS(Detector);
+			}
+		});
 	}
 
 	SECTION("Transact(Scope(UE_MT_SCOPED_WRITE_ACCESS), Scope(UE_MT_SCOPED_READ_ACCESS))")
 	{
-		FRWTransactionallySafeAccessDetector Detector;
-		AutoRTFM::ETransactionResult Result = AutoRTFM::Transact([&]
+		TestType Detector;
+		AutoRTFM::Testing::Commit([&]
+		{
 			{
-				{
-					UE_MT_SCOPED_WRITE_ACCESS(Detector);
-				}
-				{
-					UE_MT_SCOPED_READ_ACCESS(Detector);
-				}
-			});
-		REQUIRE(AutoRTFM::ETransactionResult::Committed == Result);
+				UE_MT_SCOPED_WRITE_ACCESS(Detector);
+			}
+			{
+				UE_MT_SCOPED_READ_ACCESS(Detector);
+			}
+		});
 	}
 
 	SECTION("Transact(Scope(UE_MT_SCOPED_WRITE_ACCESS), Scope(UE_MT_SCOPED_WRITE_ACCESS))")
 	{
-		FRWTransactionallySafeAccessDetector Detector;
-		AutoRTFM::ETransactionResult Result = AutoRTFM::Transact([&]
+		TestType Detector;
+		AutoRTFM::Testing::Commit([&]
+		{
 			{
-				{
-					UE_MT_SCOPED_WRITE_ACCESS(Detector);
-				}
-				{
-					UE_MT_SCOPED_WRITE_ACCESS(Detector);
-				}
-			});
-		REQUIRE(AutoRTFM::ETransactionResult::Committed == Result);
+				UE_MT_SCOPED_WRITE_ACCESS(Detector);
+			}
+			{
+				UE_MT_SCOPED_WRITE_ACCESS(Detector);
+			}
+		});
 	}
 
 	SECTION("Transact(Scope(UE_MT_SCOPED_READ_ACCESS), Scope(UE_MT_SCOPED_READ_ACCESS), Abort)")
 	{
-		FRWTransactionallySafeAccessDetector Detector;
-		AutoRTFM::ETransactionResult Result = AutoRTFM::Transact([&]
+		TestType Detector;
+		AutoRTFM::Testing::Abort([&]
+		{
 			{
-				{
-					UE_MT_SCOPED_READ_ACCESS(Detector);
-				}
-				{
-					UE_MT_SCOPED_READ_ACCESS(Detector);
-				}
-				AutoRTFM::AbortTransaction();
-			});
-		REQUIRE(AutoRTFM::ETransactionResult::AbortedByRequest == Result);
+				UE_MT_SCOPED_READ_ACCESS(Detector);
+			}
+			{
+				UE_MT_SCOPED_READ_ACCESS(Detector);
+			}
+			AutoRTFM::AbortTransaction();
+		});
 	}
 
 	SECTION("Transact(Scope(UE_MT_SCOPED_READ_ACCESS), Scope(UE_MT_SCOPED_WRITE_ACCESS), Abort)")
 	{
-		FRWTransactionallySafeAccessDetector Detector;
-		AutoRTFM::ETransactionResult Result = AutoRTFM::Transact([&]
+		TestType Detector;
+		AutoRTFM::Testing::Abort([&]
+		{
 			{
-				{
-					UE_MT_SCOPED_READ_ACCESS(Detector);
-				}
-				{
-					UE_MT_SCOPED_WRITE_ACCESS(Detector);
-				}
-				AutoRTFM::AbortTransaction();
-			});
-		REQUIRE(AutoRTFM::ETransactionResult::AbortedByRequest == Result);
+				UE_MT_SCOPED_READ_ACCESS(Detector);
+			}
+			{
+				UE_MT_SCOPED_WRITE_ACCESS(Detector);
+			}
+			AutoRTFM::AbortTransaction();
+		});
 	}
 
 	SECTION("Transact(Scope(UE_MT_SCOPED_WRITE_ACCESS), Scope(UE_MT_SCOPED_READ_ACCESS), Abort)")
 	{
-		FRWTransactionallySafeAccessDetector Detector;
-		AutoRTFM::ETransactionResult Result = AutoRTFM::Transact([&]
+		TestType Detector;
+		AutoRTFM::Testing::Abort([&]
+		{
 			{
-				{
-					UE_MT_SCOPED_WRITE_ACCESS(Detector);
-				}
-				{
-					UE_MT_SCOPED_READ_ACCESS(Detector);
-				}
-				AutoRTFM::AbortTransaction();
-			});
-		REQUIRE(AutoRTFM::ETransactionResult::AbortedByRequest == Result);
+				UE_MT_SCOPED_WRITE_ACCESS(Detector);
+			}
+			{
+				UE_MT_SCOPED_READ_ACCESS(Detector);
+			}
+			AutoRTFM::AbortTransaction();
+		});
 	}
 
 	SECTION("Transact(Scope(UE_MT_SCOPED_WRITE_ACCESS), Scope(UE_MT_SCOPED_WRITE_ACCESS), Abort)")
 	{
-		FRWTransactionallySafeAccessDetector Detector;
-		AutoRTFM::ETransactionResult Result = AutoRTFM::Transact([&]
+		TestType Detector;
+		AutoRTFM::Testing::Abort([&]
+		{
 			{
-				{
-					UE_MT_SCOPED_WRITE_ACCESS(Detector);
-				}
-				{
-					UE_MT_SCOPED_WRITE_ACCESS(Detector);
-				}
-				AutoRTFM::AbortTransaction();
-			});
-		REQUIRE(AutoRTFM::ETransactionResult::AbortedByRequest == Result);
+				UE_MT_SCOPED_WRITE_ACCESS(Detector);
+			}
+			{
+				UE_MT_SCOPED_WRITE_ACCESS(Detector);
+			}
+			AutoRTFM::AbortTransaction();
+		});
 	}
 
 	SECTION("Transact(Scope(UE_MT_SCOPED_READ_ACCESS), Scope(UE_MT_SCOPED_READ_ACCESS, Abort))")
 	{
-		FRWTransactionallySafeAccessDetector Detector;
-		AutoRTFM::ETransactionResult Result = AutoRTFM::Transact([&]
+		TestType Detector;
+		AutoRTFM::Testing::Abort([&]
+		{
 			{
-				{
-					UE_MT_SCOPED_READ_ACCESS(Detector);
-				}
-				{
-					UE_MT_SCOPED_READ_ACCESS(Detector);
-					AutoRTFM::AbortTransaction();
-				}
-			});
-		REQUIRE(AutoRTFM::ETransactionResult::AbortedByRequest == Result);
+				UE_MT_SCOPED_READ_ACCESS(Detector);
+			}
+			{
+				UE_MT_SCOPED_READ_ACCESS(Detector);
+				AutoRTFM::AbortTransaction();
+			}
+		});
 	}
 
 	SECTION("Transact(Scope(UE_MT_SCOPED_READ_ACCESS), Scope(UE_MT_SCOPED_WRITE_ACCESS, Abort))")
 	{
-		FRWTransactionallySafeAccessDetector Detector;
-		AutoRTFM::ETransactionResult Result = AutoRTFM::Transact([&]
+		TestType Detector;
+		AutoRTFM::Testing::Abort([&]
+		{
 			{
-				{
-					UE_MT_SCOPED_READ_ACCESS(Detector);
-				}
-				{
-					UE_MT_SCOPED_WRITE_ACCESS(Detector);
-					AutoRTFM::AbortTransaction();
-				}
-			});
-		REQUIRE(AutoRTFM::ETransactionResult::AbortedByRequest == Result);
+				UE_MT_SCOPED_READ_ACCESS(Detector);
+			}
+			{
+				UE_MT_SCOPED_WRITE_ACCESS(Detector);
+				AutoRTFM::AbortTransaction();
+			}
+		});
 	}
 
 	SECTION("Transact(Scope(UE_MT_SCOPED_WRITE_ACCESS), Scope(UE_MT_SCOPED_READ_ACCESS, Abort))")
 	{
-		FRWTransactionallySafeAccessDetector Detector;
-		AutoRTFM::ETransactionResult Result = AutoRTFM::Transact([&]
+		TestType Detector;
+		AutoRTFM::Testing::Abort([&]
+		{
 			{
-				{
-					UE_MT_SCOPED_WRITE_ACCESS(Detector);
-				}
-				{
-					UE_MT_SCOPED_READ_ACCESS(Detector);
-					AutoRTFM::AbortTransaction();
-				}
-			});
-		REQUIRE(AutoRTFM::ETransactionResult::AbortedByRequest == Result);
+				UE_MT_SCOPED_WRITE_ACCESS(Detector);
+			}
+			{
+				UE_MT_SCOPED_READ_ACCESS(Detector);
+				AutoRTFM::AbortTransaction();
+			}
+		});
 	}
 
 	SECTION("Transact(Scope(UE_MT_SCOPED_WRITE_ACCESS), Scope(UE_MT_SCOPED_WRITE_ACCESS, Abort))")
 	{
-		FRWTransactionallySafeAccessDetector Detector;
-		AutoRTFM::ETransactionResult Result = AutoRTFM::Transact([&]
+		TestType Detector;
+		AutoRTFM::Testing::Abort([&]
+		{
 			{
-				{
-					UE_MT_SCOPED_WRITE_ACCESS(Detector);
-				}
-				{
-					UE_MT_SCOPED_WRITE_ACCESS(Detector);
-					AutoRTFM::AbortTransaction();
-				}
-			});
-		REQUIRE(AutoRTFM::ETransactionResult::AbortedByRequest == Result);
+				UE_MT_SCOPED_WRITE_ACCESS(Detector);
+			}
+			{
+				UE_MT_SCOPED_WRITE_ACCESS(Detector);
+				AutoRTFM::AbortTransaction();
+			}
+		});
 	}
 
 	SECTION("Transact(UE_MT_SCOPED_READ_ACCESS, Open(UE_MT_SCOPED_READ_ACCESS))")
 	{
-		FRWTransactionallySafeAccessDetector Detector;
-		AutoRTFM::ETransactionResult Result = AutoRTFM::Transact([&]
+		TestType Detector;
+		AutoRTFM::Testing::Commit([&]
+		{
+			UE_MT_SCOPED_READ_ACCESS(Detector);
+			AutoRTFM::Open([&]
 			{
 				UE_MT_SCOPED_READ_ACCESS(Detector);
-				AutoRTFM::Open([&]
-					{
-						UE_MT_SCOPED_READ_ACCESS(Detector);
-					});
 			});
-		REQUIRE(AutoRTFM::ETransactionResult::Committed == Result);
+		});
 	}
 
 	SECTION("Transact(UE_MT_SCOPED_READ_ACCESS, Open(UE_MT_SCOPED_READ_ACCESS), Abort)")
 	{
-		FRWTransactionallySafeAccessDetector Detector;
-		AutoRTFM::ETransactionResult Result = AutoRTFM::Transact([&]
+		TestType Detector;
+		AutoRTFM::Testing::Abort([&]
+		{
+			UE_MT_SCOPED_READ_ACCESS(Detector);
+			AutoRTFM::Open([&]
 			{
 				UE_MT_SCOPED_READ_ACCESS(Detector);
-				AutoRTFM::Open([&]
-					{
-						UE_MT_SCOPED_READ_ACCESS(Detector);
-					});
-				AutoRTFM::AbortTransaction();
 			});
-		REQUIRE(AutoRTFM::ETransactionResult::AbortedByRequest == Result);
+			AutoRTFM::AbortTransaction();
+		});
 	}
 
 	SECTION("Transact(Scope(UE_MT_SCOPED_READ_ACCESS, Open(UE_MT_SCOPED_READ_ACCESS)), Abort)")
 	{
-		FRWTransactionallySafeAccessDetector Detector;
-		AutoRTFM::ETransactionResult Result = AutoRTFM::Transact([&]
+		TestType Detector;
+		AutoRTFM::Testing::Abort([&]
+		{
 			{
+				UE_MT_SCOPED_READ_ACCESS(Detector);
+				AutoRTFM::Open([&]
 				{
 					UE_MT_SCOPED_READ_ACCESS(Detector);
-					AutoRTFM::Open([&]
-						{
-							UE_MT_SCOPED_READ_ACCESS(Detector);
-						});
-				}
-				AutoRTFM::AbortTransaction();
-			});
-		REQUIRE(AutoRTFM::ETransactionResult::AbortedByRequest == Result);
+				});
+			}
+			AutoRTFM::AbortTransaction();
+		});
 	}
 
 	SECTION("Transact(Scope(UE_MT_SCOPED_READ_ACCESS), Open(UE_MT_SCOPED_READ_ACCESS))")
 	{
-		FRWTransactionallySafeAccessDetector Detector;
-		AutoRTFM::ETransactionResult Result = AutoRTFM::Transact([&]
+		TestType Detector;
+		AutoRTFM::Testing::Commit([&]
+		{
 			{
-				{
-					UE_MT_SCOPED_READ_ACCESS(Detector);
-				}
-				AutoRTFM::Open([&]
-					{
-						UE_MT_SCOPED_READ_ACCESS(Detector);
-					});
+				UE_MT_SCOPED_READ_ACCESS(Detector);
+			}
+			AutoRTFM::Open([&]
+			{
+				UE_MT_SCOPED_READ_ACCESS(Detector);
 			});
-		REQUIRE(AutoRTFM::ETransactionResult::Committed == Result);
+		});
 	}
 
 	SECTION("Transact(Scope(UE_MT_SCOPED_READ_ACCESS), Open(UE_MT_SCOPED_WRITE_ACCESS))")
 	{
-		FRWTransactionallySafeAccessDetector Detector;
-		AutoRTFM::ETransactionResult Result = AutoRTFM::Transact([&]
+		TestType Detector;
+		AutoRTFM::Testing::Commit([&]
+		{
 			{
-				{
-					UE_MT_SCOPED_READ_ACCESS(Detector);
-				}
-				AutoRTFM::Open([&]
-					{
-						UE_MT_SCOPED_WRITE_ACCESS(Detector);
-					});
+				UE_MT_SCOPED_READ_ACCESS(Detector);
+			}
+			AutoRTFM::Open([&]
+			{
+				UE_MT_SCOPED_WRITE_ACCESS(Detector);
 			});
-		REQUIRE(AutoRTFM::ETransactionResult::Committed == Result);
+		});
 	}
 
 	SECTION("Transact(Scope(UE_MT_SCOPED_WRITE_ACCESS), Open(UE_MT_SCOPED_READ_ACCESS))")
 	{
-		FRWTransactionallySafeAccessDetector Detector;
-		AutoRTFM::ETransactionResult Result = AutoRTFM::Transact([&]
+		TestType Detector;
+		AutoRTFM::Testing::Commit([&]
+		{
 			{
-				{
-					UE_MT_SCOPED_WRITE_ACCESS(Detector);
-				}
-				AutoRTFM::Open([&]
-					{
-						UE_MT_SCOPED_READ_ACCESS(Detector);
-					});
+				UE_MT_SCOPED_WRITE_ACCESS(Detector);
+			}
+			AutoRTFM::Open([&]
+			{
+				UE_MT_SCOPED_READ_ACCESS(Detector);
 			});
-		REQUIRE(AutoRTFM::ETransactionResult::Committed == Result);
+		});
 	}
 
 	SECTION("Transact(Scope(UE_MT_SCOPED_WRITE_ACCESS), Open(UE_MT_SCOPED_WRITE_ACCESS))")
 	{
-		FRWTransactionallySafeAccessDetector Detector;
-		AutoRTFM::ETransactionResult Result = AutoRTFM::Transact([&]
+		TestType Detector;
+		AutoRTFM::Testing::Commit([&]
+		{
 			{
-				{
-					UE_MT_SCOPED_WRITE_ACCESS(Detector);
-				}
-				AutoRTFM::Open([&]
-					{
-						UE_MT_SCOPED_WRITE_ACCESS(Detector);
-					});
+				UE_MT_SCOPED_WRITE_ACCESS(Detector);
+			}
+			AutoRTFM::Open([&]
+			{
+				UE_MT_SCOPED_WRITE_ACCESS(Detector);
 			});
-		REQUIRE(AutoRTFM::ETransactionResult::Committed == Result);
+		});
 	}
 
 	SECTION("Transact(Open(UE_MT_SCOPED_READ_ACCESS), Scope(UE_MT_SCOPED_READ_ACCESS))")
 	{
-		FRWTransactionallySafeAccessDetector Detector;
-		AutoRTFM::ETransactionResult Result = AutoRTFM::Transact([&]
+		TestType Detector;
+		AutoRTFM::Testing::Commit([&]
+		{
+			AutoRTFM::Open([&]
 			{
-				AutoRTFM::Open([&]
-					{
-						UE_MT_SCOPED_READ_ACCESS(Detector);
-					});
 				UE_MT_SCOPED_READ_ACCESS(Detector);
 			});
-		REQUIRE(AutoRTFM::ETransactionResult::Committed == Result);
+			UE_MT_SCOPED_READ_ACCESS(Detector);
+		});
 	}
 
 	SECTION("Transact(Open(UE_MT_SCOPED_READ_ACCESS), Scope(UE_MT_SCOPED_WRITE_ACCESS))")
 	{
-		FRWTransactionallySafeAccessDetector Detector;
-		AutoRTFM::ETransactionResult Result = AutoRTFM::Transact([&]
+		TestType Detector;
+		AutoRTFM::Testing::Commit([&]
+		{
+			AutoRTFM::Open([&]
 			{
-				AutoRTFM::Open([&]
-					{
-						UE_MT_SCOPED_READ_ACCESS(Detector);
-					});
-				UE_MT_SCOPED_WRITE_ACCESS(Detector);
+				UE_MT_SCOPED_READ_ACCESS(Detector);
 			});
-		REQUIRE(AutoRTFM::ETransactionResult::Committed == Result);
+			UE_MT_SCOPED_WRITE_ACCESS(Detector);
+		});
 	}
 
 	SECTION("Transact(Open(UE_MT_SCOPED_WRITE_ACCESS), Scope(UE_MT_SCOPED_READ_ACCESS))")
 	{
-		FRWTransactionallySafeAccessDetector Detector;
-		AutoRTFM::ETransactionResult Result = AutoRTFM::Transact([&]
+		TestType Detector;
+		AutoRTFM::Testing::Commit([&]
+		{
+			AutoRTFM::Open([&]
 			{
-				AutoRTFM::Open([&]
-					{
-						UE_MT_SCOPED_WRITE_ACCESS(Detector);
-					});
-				UE_MT_SCOPED_READ_ACCESS(Detector);
+				UE_MT_SCOPED_WRITE_ACCESS(Detector);
 			});
-		REQUIRE(AutoRTFM::ETransactionResult::Committed == Result);
+			UE_MT_SCOPED_READ_ACCESS(Detector);
+		});
 	}
 
 	SECTION("Transact(Open(UE_MT_SCOPED_WRITE_ACCESS), UE_MT_SCOPED_WRITE_ACCESS)")
 	{
-		FRWTransactionallySafeAccessDetector Detector;
-		AutoRTFM::ETransactionResult Result = AutoRTFM::Transact([&]
+		TestType Detector;
+		AutoRTFM::Testing::Commit([&]
+		{
+			AutoRTFM::Open([&]
 			{
-				AutoRTFM::Open([&]
-					{
-						UE_MT_SCOPED_WRITE_ACCESS(Detector);
-					});
 				UE_MT_SCOPED_WRITE_ACCESS(Detector);
 			});
-		REQUIRE(AutoRTFM::ETransactionResult::Committed == Result);
+			UE_MT_SCOPED_WRITE_ACCESS(Detector);
+		});
 	}
 
 	SECTION("Transact(Scope(UE_MT_SCOPED_READ_ACCESS, Spawn(UE_MT_SCOPED_READ_ACCESS)))")
 	{
-		FRWTransactionallySafeAccessDetector Detector;
+		TestType Detector;
 		std::atomic_uint Handshake = 0; // Used to order this thread and the spawnee.
 		std::thread OtherThread;
-		AutoRTFM::ETransactionResult Result = AutoRTFM::Transact([&]
+		AutoRTFM::Testing::Commit([&]
+		{
+			// Because we can retry transactions reset handshake to zero.
+			AutoRTFM::Open([&]
 			{
-				// Because we can retry transactions reset handshake to zero.
-				AutoRTFM::Open([&]
-					{
-						Handshake = 0;
-					});
+				Handshake = 0;
+			});
 
-				{
-					UE_MT_SCOPED_READ_ACCESS(Detector);
+			{
+				UE_MT_SCOPED_READ_ACCESS(Detector);
 
-					AutoRTFM::Open([&]
+				AutoRTFM::Open([&]
+				{
+					OtherThread = std::thread([&]
 						{
-							OtherThread = std::thread([&]
-								{
-									// Let the main thread progress.
-									Handshake += 1;
-
-									{
-										UE_MT_SCOPED_READ_ACCESS(Detector);
-
-										// Wait for the main thread.
-										while (2 != Handshake) {}
-									}
-								});
+							// Let the main thread progress.
+							Handshake += 1;
+
+							{
+								UE_MT_SCOPED_READ_ACCESS(Detector);
+
+								// Wait for the main thread.
+								while (2 != Handshake) {}
+							}
 						});
-				}
+				});
+			}
 
-				AutoRTFM::Open([&]
-					{
-						// Wait for the spawnee.
-						while (1 != Handshake) {}
+			AutoRTFM::Open([&]
+			{
+				// Wait for the spawnee.
+				while (1 != Handshake) {}
 
-						// Let the spawnee progress.
-						Handshake += 1;
+				// Let the spawnee progress.
+				Handshake += 1;
 
-						// Wait for the spawnee to finish before unwinding the
-						// stack and invalidating the Handshake and Detector
-						// references used by the thread.
-						OtherThread.join();
-					});
+				// Wait for the spawnee to finish before unwinding the
+				// stack and invalidating the Handshake and Detector
+				// references used by the thread.
+				OtherThread.join();
 			});
-		REQUIRE(AutoRTFM::ETransactionResult::Committed == Result);
+		});
 	}
 
 	SECTION("Transact(UE_MT_SCOPED_READ_ACCESS, Transact(UE_MT_SCOPED_READ_ACCESS))")
 	{
-		FRWTransactionallySafeAccessDetector Detector;
-		AutoRTFM::ETransactionResult Result = AutoRTFM::Transact([&]
+		TestType Detector;
+		AutoRTFM::Testing::Commit([&]
+		{
+			UE_MT_SCOPED_READ_ACCESS(Detector);
+
+			AutoRTFM::Testing::Commit([&]
 			{
 				UE_MT_SCOPED_READ_ACCESS(Detector);
-
-				AutoRTFM::Commit([&]
-					{
-						UE_MT_SCOPED_READ_ACCESS(Detector);
-					});
 			});
-		REQUIRE(AutoRTFM::ETransactionResult::Committed == Result);
+		});
 	}
 
 	SECTION("Transact(Scope(UE_MT_SCOPED_READ_ACCESS), Transact(UE_MT_SCOPED_READ_ACCESS))")
 	{
-		FRWTransactionallySafeAccessDetector Detector;
-		AutoRTFM::ETransactionResult Result = AutoRTFM::Transact([&]
+		TestType Detector;
+		AutoRTFM::Testing::Commit([&]
+		{
 			{
-				{
-					UE_MT_SCOPED_READ_ACCESS(Detector);
-				}
+				UE_MT_SCOPED_READ_ACCESS(Detector);
+			}
 
-				AutoRTFM::Commit([&]
-					{
-						UE_MT_SCOPED_READ_ACCESS(Detector);
-					});
+			AutoRTFM::Testing::Commit([&]
+			{
+				UE_MT_SCOPED_READ_ACCESS(Detector);
 			});
-		REQUIRE(AutoRTFM::ETransactionResult::Committed == Result);
+		});
 	}
 
 	SECTION("Transact(Scope(UE_MT_SCOPED_READ_ACCESS), Transact(UE_MT_SCOPED_WRITE_ACCESS))")
 	{
-		FRWTransactionallySafeAccessDetector Detector;
-		AutoRTFM::ETransactionResult Result = AutoRTFM::Transact([&]
+		TestType Detector;
+		AutoRTFM::Testing::Commit([&]
+		{
 			{
-				{
-					UE_MT_SCOPED_READ_ACCESS(Detector);
-				}
+				UE_MT_SCOPED_READ_ACCESS(Detector);
+			}
 
-				AutoRTFM::Commit([&]
-					{
-						UE_MT_SCOPED_WRITE_ACCESS(Detector);
-					});
+			AutoRTFM::Testing::Commit([&]
+			{
+				UE_MT_SCOPED_WRITE_ACCESS(Detector);
 			});
-		REQUIRE(AutoRTFM::ETransactionResult::Committed == Result);
+		});
 	}
 
 	SECTION("Transact(Scope(UE_MT_SCOPED_WRITE_ACCESS), Transact(UE_MT_SCOPED_READ_ACCESS))")
 	{
-		FRWTransactionallySafeAccessDetector Detector;
-		AutoRTFM::ETransactionResult Result = AutoRTFM::Transact([&]
+		TestType Detector;
+		AutoRTFM::Testing::Commit([&]
+		{
 			{
-				{
-					UE_MT_SCOPED_WRITE_ACCESS(Detector);
-				}
+				UE_MT_SCOPED_WRITE_ACCESS(Detector);
+			}
 
-				AutoRTFM::Commit([&]
-					{
-						UE_MT_SCOPED_READ_ACCESS(Detector);
-					});
+			AutoRTFM::Testing::Commit([&]
+			{
+				UE_MT_SCOPED_READ_ACCESS(Detector);
 			});
-		REQUIRE(AutoRTFM::ETransactionResult::Committed == Result);
+		});
 	}
 
 	SECTION("Transact(Scope(UE_MT_SCOPED_WRITE_ACCESS), Transact(UE_MT_SCOPED_WRITE_ACCESS))")
 	{
-		FRWTransactionallySafeAccessDetector Detector;
-		AutoRTFM::ETransactionResult Result = AutoRTFM::Transact([&]
+		TestType Detector;
+		AutoRTFM::Testing::Commit([&]
+		{
 			{
-				{
-					UE_MT_SCOPED_WRITE_ACCESS(Detector);
-				}
+				UE_MT_SCOPED_WRITE_ACCESS(Detector);
+			}
 
-				AutoRTFM::Commit([&]
-					{
-						UE_MT_SCOPED_WRITE_ACCESS(Detector);
-					});
+			AutoRTFM::Testing::Commit([&]
+			{
+				UE_MT_SCOPED_WRITE_ACCESS(Detector);
 			});
-		REQUIRE(AutoRTFM::ETransactionResult::Committed == Result);
+		});
 	}
 
 	SECTION("Transact(Scope(UE_MT_SCOPED_READ_ACCESS), Transact(UE_MT_SCOPED_READ_ACCESS), Abort)")
 	{
-		FRWTransactionallySafeAccessDetector Detector;
-		AutoRTFM::ETransactionResult Result = AutoRTFM::Transact([&]
+		TestType Detector;
+		AutoRTFM::Testing::Abort([&]
+		{
 			{
-				{
-					UE_MT_SCOPED_READ_ACCESS(Detector);
-				}
-
-				AutoRTFM::Commit([&]
-					{
-						UE_MT_SCOPED_READ_ACCESS(Detector);
-					});
+				UE_MT_SCOPED_READ_ACCESS(Detector);
+			}
 
-				AutoRTFM::AbortTransaction();
+			AutoRTFM::Testing::Commit([&]
+			{
+				UE_MT_SCOPED_READ_ACCESS(Detector);
 			});
-		REQUIRE(AutoRTFM::ETransactionResult::AbortedByRequest == Result);
+
+			AutoRTFM::AbortTransaction();
+		});
 	}
 
 	SECTION("Transact(Scope(UE_MT_SCOPED_READ_ACCESS), Transact(UE_MT_SCOPED_WRITE_ACCESS), Abort)")
 	{
-		FRWTransactionallySafeAccessDetector Detector;
-		AutoRTFM::ETransactionResult Result = AutoRTFM::Transact([&]
+		TestType Detector;
+		AutoRTFM::Testing::Abort([&]
+		{
 			{
-				{
-					UE_MT_SCOPED_READ_ACCESS(Detector);
-				}
-
-				AutoRTFM::Commit([&]
-					{
-						UE_MT_SCOPED_WRITE_ACCESS(Detector);
-					});
+				UE_MT_SCOPED_READ_ACCESS(Detector);
+			}
 
-				AutoRTFM::AbortTransaction();
+			AutoRTFM::Testing::Commit([&]
+			{
+				UE_MT_SCOPED_WRITE_ACCESS(Detector);
 			});
-		REQUIRE(AutoRTFM::ETransactionResult::AbortedByRequest == Result);
+
+			AutoRTFM::AbortTransaction();
+		});
 	}
 
 	SECTION("Transact(Scope(UE_MT_SCOPED_WRITE_ACCESS), Transact(UE_MT_SCOPED_READ_ACCESS), Abort)")
 	{
-		FRWTransactionallySafeAccessDetector Detector;
-		AutoRTFM::ETransactionResult Result = AutoRTFM::Transact([&]
+		TestType Detector;
+		AutoRTFM::Testing::Abort([&]
+		{
 			{
-				{
-					UE_MT_SCOPED_WRITE_ACCESS(Detector);
-				}
-
-				AutoRTFM::Commit([&]
-					{
-						UE_MT_SCOPED_READ_ACCESS(Detector);
-					});
+				UE_MT_SCOPED_WRITE_ACCESS(Detector);
+			}
 
-				AutoRTFM::AbortTransaction();
+			AutoRTFM::Testing::Commit([&]
+			{
+				UE_MT_SCOPED_READ_ACCESS(Detector);
 			});
-		REQUIRE(AutoRTFM::ETransactionResult::AbortedByRequest == Result);
+
+			AutoRTFM::AbortTransaction();
+		});
 	}
 
 	SECTION("Transact(Scope(UE_MT_SCOPED_WRITE_ACCESS), Transact(UE_MT_SCOPED_WRITE_ACCESS), Abort)")
 	{
-		FRWTransactionallySafeAccessDetector Detector;
-		AutoRTFM::ETransactionResult Result = AutoRTFM::Transact([&]
+		TestType Detector;
+		AutoRTFM::Testing::Abort([&]
+		{
 			{
-				{
-					UE_MT_SCOPED_WRITE_ACCESS(Detector);
-				}
-
-				AutoRTFM::Commit([&]
-					{
-						UE_MT_SCOPED_WRITE_ACCESS(Detector);
-					});
+				UE_MT_SCOPED_WRITE_ACCESS(Detector);
+			}
 
-				AutoRTFM::AbortTransaction();
+			AutoRTFM::Testing::Commit([&]
+			{
+				UE_MT_SCOPED_WRITE_ACCESS(Detector);
 			});
-		REQUIRE(AutoRTFM::ETransactionResult::AbortedByRequest == Result);
+
+			AutoRTFM::AbortTransaction();
+		});
 	}
 
-	SECTION("Transact(Scope(UE_MT_SCOPED_READ_ACCESS), Transact(UE_MT_SCOPED_READ_ACCESS), Abort)")
+	SECTION("Transact(Scope(UE_MT_SCOPED_READ_ACCESS), Transact(UE_MT_SCOPED_READ_ACCESS, Abort))")
 	{
-		FRWTransactionallySafeAccessDetector Detector;
-		AutoRTFM::ETransactionResult Result = AutoRTFM::Transact([&]
+		TestType Detector;
+		AutoRTFM::Testing::Commit([&]
+		{
 			{
-				{
-					UE_MT_SCOPED_READ_ACCESS(Detector);
-				}
+				UE_MT_SCOPED_READ_ACCESS(Detector);
+			}
 
-				AutoRTFM::ETransactionResult Result2 = AutoRTFM::Transact([&]
-					{
-						UE_MT_SCOPED_READ_ACCESS(Detector);
-						AutoRTFM::AbortTransaction();
-					});
-				REQUIRE(AutoRTFM::ETransactionResult::AbortedByRequest == Result2);
+			AutoRTFM::Testing::Abort([&]
+			{
+				UE_MT_SCOPED_READ_ACCESS(Detector);
+				AutoRTFM::AbortTransaction();
 			});
-		REQUIRE(AutoRTFM::ETransactionResult::Committed == Result);
+		});
 	}
 
-	SECTION("Transact(Scope(UE_MT_SCOPED_READ_ACCESS), Transact(UE_MT_SCOPED_WRITE_ACCESS), Abort)")
+	SECTION("Transact(Scope(UE_MT_SCOPED_READ_ACCESS), Transact(UE_MT_SCOPED_WRITE_ACCESS, Abort))")
 	{
-		FRWTransactionallySafeAccessDetector Detector;
-		AutoRTFM::ETransactionResult Result = AutoRTFM::Transact([&]
+		TestType Detector;
+		AutoRTFM::Testing::Commit([&]
+		{
 			{
-				{
-					UE_MT_SCOPED_READ_ACCESS(Detector);
-				}
+				UE_MT_SCOPED_READ_ACCESS(Detector);
+			}
 
-				AutoRTFM::ETransactionResult Result2 = AutoRTFM::Transact([&]
-					{
-						UE_MT_SCOPED_WRITE_ACCESS(Detector);
-						AutoRTFM::AbortTransaction();
-					});
-				REQUIRE(AutoRTFM::ETransactionResult::AbortedByRequest == Result2);
+			AutoRTFM::Testing::Abort([&]
+			{
+				UE_MT_SCOPED_WRITE_ACCESS(Detector);
+				AutoRTFM::AbortTransaction();
 			});
-		REQUIRE(AutoRTFM::ETransactionResult::Committed == Result);
+		});
 	}
 
-	SECTION("Transact(Scope(UE_MT_SCOPED_WRITE_ACCESS), Transact(UE_MT_SCOPED_READ_ACCESS), Abort)")
+	SECTION("Transact(Scope(UE_MT_SCOPED_WRITE_ACCESS), Transact(UE_MT_SCOPED_READ_ACCESS, Abort))")
 	{
-		FRWTransactionallySafeAccessDetector Detector;
-		AutoRTFM::ETransactionResult Result = AutoRTFM::Transact([&]
+		TestType Detector;
+		AutoRTFM::Testing::Commit([&]
+		{
 			{
-				{
-					UE_MT_SCOPED_WRITE_ACCESS(Detector);
-				}
+				UE_MT_SCOPED_WRITE_ACCESS(Detector);
+			}
 
-				AutoRTFM::ETransactionResult Result2 = AutoRTFM::Transact([&]
-					{
-						UE_MT_SCOPED_READ_ACCESS(Detector);
-						AutoRTFM::AbortTransaction();
-					});
-				REQUIRE(AutoRTFM::ETransactionResult::AbortedByRequest == Result2);
+			AutoRTFM::Testing::Abort([&]
+			{
+				UE_MT_SCOPED_READ_ACCESS(Detector);
+				AutoRTFM::AbortTransaction();
 			});
-		REQUIRE(AutoRTFM::ETransactionResult::Committed == Result);
+		});
 	}
 
-	SECTION("Transact(Scope(UE_MT_SCOPED_WRITE_ACCESS), Transact(UE_MT_SCOPED_WRITE_ACCESS), Abort)")
+	SECTION("Transact(Scope(UE_MT_SCOPED_WRITE_ACCESS), Transact(UE_MT_SCOPED_WRITE_ACCESS, Abort))")
 	{
-		FRWTransactionallySafeAccessDetector Detector;
-		AutoRTFM::ETransactionResult Result = AutoRTFM::Transact([&]
+		TestType Detector;
+		AutoRTFM::Testing::Commit([&]
+		{
 			{
-				{
-					UE_MT_SCOPED_WRITE_ACCESS(Detector);
-				}
+				UE_MT_SCOPED_WRITE_ACCESS(Detector);
+			}
 
-				AutoRTFM::ETransactionResult Result2 = AutoRTFM::Transact([&]
-					{
-						UE_MT_SCOPED_WRITE_ACCESS(Detector);
-						AutoRTFM::AbortTransaction();
-					});
-				REQUIRE(AutoRTFM::ETransactionResult::AbortedByRequest == Result2);
+			AutoRTFM::Testing::Abort([&]
+			{
+				UE_MT_SCOPED_WRITE_ACCESS(Detector);
+				AutoRTFM::AbortTransaction();
 			});
-		REQUIRE(AutoRTFM::ETransactionResult::Committed == Result);
+		});
 	}
 
 	SECTION("Transact(Scope(UE_MT_SCOPED_READ_ACCESS), OnCommit(UE_MT_SCOPED_READ_ACCESS))")
 	{
-		FRWTransactionallySafeAccessDetector Detector;
-		AutoRTFM::ETransactionResult Result = AutoRTFM::Transact([&]
+		TestType Detector;
+		AutoRTFM::Testing::Commit([&]
+		{
 			{
-				{
-					UE_MT_SCOPED_READ_ACCESS(Detector);
-				}
+				UE_MT_SCOPED_READ_ACCESS(Detector);
+			}
 
-				AutoRTFM::OnCommit([&]
-					{
-						UE_MT_SCOPED_READ_ACCESS(Detector);
-					});
+			AutoRTFM::OnCommit([&]
+			{
+				UE_MT_SCOPED_READ_ACCESS(Detector);
 			});
-		REQUIRE(AutoRTFM::ETransactionResult::Committed == Result);
+		});
 	}
 
 	SECTION("Transact(Scope(UE_MT_SCOPED_READ_ACCESS), OnCommit(UE_MT_SCOPED_WRITE_ACCESS))")
 	{
-		FRWTransactionallySafeAccessDetector Detector;
-		AutoRTFM::ETransactionResult Result = AutoRTFM::Transact([&]
+		TestType Detector;
+		AutoRTFM::Testing::Commit([&]
+		{
 			{
-				{
-					UE_MT_SCOPED_READ_ACCESS(Detector);
-				}
+				UE_MT_SCOPED_READ_ACCESS(Detector);
+			}
 
-				AutoRTFM::OnCommit([&]
-					{
-						UE_MT_SCOPED_WRITE_ACCESS(Detector);
-					});
+			AutoRTFM::OnCommit([&]
+			{
+				UE_MT_SCOPED_WRITE_ACCESS(Detector);
 			});
-		REQUIRE(AutoRTFM::ETransactionResult::Committed == Result);
+		});
 	}
 
 	SECTION("Transact(Scope(UE_MT_SCOPED_WRITE_ACCESS), OnCommit(UE_MT_SCOPED_READ_ACCESS))")
 	{
-		FRWTransactionallySafeAccessDetector Detector;
-		AutoRTFM::ETransactionResult Result = AutoRTFM::Transact([&]
+		TestType Detector;
+		AutoRTFM::Testing::Commit([&]
+		{
 			{
-				{
-					UE_MT_SCOPED_WRITE_ACCESS(Detector);
-				}
+				UE_MT_SCOPED_WRITE_ACCESS(Detector);
+			}
 
-				AutoRTFM::OnCommit([&]
-					{
-						UE_MT_SCOPED_READ_ACCESS(Detector);
-					});
+			AutoRTFM::OnCommit([&]
+			{
+				UE_MT_SCOPED_READ_ACCESS(Detector);
 			});
-		REQUIRE(AutoRTFM::ETransactionResult::Committed == Result);
+		});
 	}
 
 	SECTION("Transact(Scope(UE_MT_SCOPED_WRITE_ACCESS), OnCommit(UE_MT_SCOPED_WRITE_ACCESS))")
 	{
-		FRWTransactionallySafeAccessDetector Detector;
-		AutoRTFM::ETransactionResult Result = AutoRTFM::Transact([&]
+		TestType Detector;
+		AutoRTFM::Testing::Commit([&]
+		{
 			{
-				{
-					UE_MT_SCOPED_WRITE_ACCESS(Detector);
-				}
+				UE_MT_SCOPED_WRITE_ACCESS(Detector);
+			}
 
-				AutoRTFM::OnCommit([&]
-					{
-						UE_MT_SCOPED_WRITE_ACCESS(Detector);
-					});
+			AutoRTFM::OnCommit([&]
+			{
+				UE_MT_SCOPED_WRITE_ACCESS(Detector);
 			});
-		REQUIRE(AutoRTFM::ETransactionResult::Committed == Result);
+		});
 	}
 
-	SECTION("Transact(Scope(UE_MT_SCOPED_READ_ACCESS), OnAbort(UE_MT_SCOPED_READ_ACCESS))")
+	SECTION("Transact(Scope(UE_MT_SCOPED_READ_ACCESS), OnAbort(UE_MT_SCOPED_READ_ACCESS), Abort)")
 	{
-		FRWTransactionallySafeAccessDetector Detector;
-		AutoRTFM::ETransactionResult Result = AutoRTFM::Transact([&]
+		TestType Detector;
+		AutoRTFM::Testing::Abort([&]
+		{
 			{
-				{
-					UE_MT_SCOPED_READ_ACCESS(Detector);
-				}
-
-				AutoRTFM::OnAbort([&]
-					{
-						UE_MT_SCOPED_READ_ACCESS(Detector);
-					});
+				UE_MT_SCOPED_READ_ACCESS(Detector);
+			}
 
-				AutoRTFM::AbortTransaction();
+			AutoRTFM::OnAbort([&]
+			{
+				UE_MT_SCOPED_READ_ACCESS(Detector);
 			});
-		REQUIRE(AutoRTFM::ETransactionResult::AbortedByRequest == Result);
+
+			AutoRTFM::AbortTransaction();
+		});
 	}
 
-	SECTION("Transact(Scope(UE_MT_SCOPED_READ_ACCESS), OnAbort(UE_MT_SCOPED_WRITE_ACCESS))")
+	SECTION("Transact(Scope(UE_MT_SCOPED_READ_ACCESS), OnAbort(UE_MT_SCOPED_WRITE_ACCESS), Abort)")
 	{
-		FRWTransactionallySafeAccessDetector Detector;
-		AutoRTFM::ETransactionResult Result = AutoRTFM::Transact([&]
+		TestType Detector;
+		AutoRTFM::Testing::Abort([&]
+		{
 			{
-				{
-					UE_MT_SCOPED_READ_ACCESS(Detector);
-				}
-
-				AutoRTFM::OnAbort([&]
-					{
-						UE_MT_SCOPED_WRITE_ACCESS(Detector);
-					});
+				UE_MT_SCOPED_READ_ACCESS(Detector);
+			}
 
-				AutoRTFM::AbortTransaction();
+			AutoRTFM::OnAbort([&]
+			{
+				UE_MT_SCOPED_WRITE_ACCESS(Detector);
 			});
-		REQUIRE(AutoRTFM::ETransactionResult::AbortedByRequest == Result);
+
+			AutoRTFM::AbortTransaction();
+		});
 	}
 
-	SECTION("Transact(Scope(UE_MT_SCOPED_WRITE_ACCESS), OnAbort(UE_MT_SCOPED_READ_ACCESS))")
+	SECTION("Transact(Scope(UE_MT_SCOPED_WRITE_ACCESS), OnAbort(UE_MT_SCOPED_READ_ACCESS), Abort)")
 	{
-		FRWTransactionallySafeAccessDetector Detector;
-		AutoRTFM::ETransactionResult Result = AutoRTFM::Transact([&]
+		TestType Detector;
+		AutoRTFM::Testing::Abort([&]
+		{
 			{
-				{
-					UE_MT_SCOPED_WRITE_ACCESS(Detector);
-				}
-
-				AutoRTFM::OnAbort([&]
-					{
-						UE_MT_SCOPED_READ_ACCESS(Detector);
-					});
+				UE_MT_SCOPED_WRITE_ACCESS(Detector);
+			}
 
-				AutoRTFM::AbortTransaction();
+			AutoRTFM::OnAbort([&]
+			{
+				UE_MT_SCOPED_READ_ACCESS(Detector);
 			});
-		REQUIRE(AutoRTFM::ETransactionResult::AbortedByRequest == Result);
+
+			AutoRTFM::AbortTransaction();
+		});
 	}
 
 	SECTION("Transact(Scope(UE_MT_SCOPED_WRITE_ACCESS), OnAbort(UE_MT_SCOPED_WRITE_ACCESS), Abort)")
 	{
-		FRWTransactionallySafeAccessDetector Detector;
-		AutoRTFM::ETransactionResult Result = AutoRTFM::Transact([&]
+		TestType Detector;
+		AutoRTFM::Testing::Abort([&]
+		{
 			{
-				{
-					UE_MT_SCOPED_WRITE_ACCESS(Detector);
-				}
-
-				AutoRTFM::OnAbort([&]
-					{
-						UE_MT_SCOPED_WRITE_ACCESS(Detector);
-					});
+				UE_MT_SCOPED_WRITE_ACCESS(Detector);
+			}
 
-				AutoRTFM::AbortTransaction();
+			AutoRTFM::OnAbort([&]
+			{
+				UE_MT_SCOPED_WRITE_ACCESS(Detector);
 			});
-		REQUIRE(AutoRTFM::ETransactionResult::AbortedByRequest == Result);
+
+			AutoRTFM::AbortTransaction();
+		});
 	}
 
 	SECTION("Transact(Create, UE_MT_SCOPED_READ_ACCESS)")
 	{
-		AutoRTFM::ETransactionResult Result = AutoRTFM::Transact([&]
-			{
-				FRWTransactionallySafeAccessDetector Detector;
-				UE_MT_SCOPED_READ_ACCESS(Detector);
-			});
-		REQUIRE(AutoRTFM::ETransactionResult::Committed == Result);
+		AutoRTFM::Testing::Commit([&]
+		{
+			TestType Detector;
+			UE_MT_SCOPED_READ_ACCESS(Detector);
+		});
 	}
 
 	SECTION("Transact(Create, UE_MT_SCOPED_READ_ACCESS, Abort)")
 	{
-		AutoRTFM::ETransactionResult Result = AutoRTFM::Transact([&]
-			{
-				FRWTransactionallySafeAccessDetector Detector;
-				UE_MT_SCOPED_READ_ACCESS(Detector);
-				AutoRTFM::AbortTransaction();
-			});
-		REQUIRE(AutoRTFM::ETransactionResult::AbortedByRequest == Result);
+		AutoRTFM::Testing::Abort([&]
+		{
+			TestType Detector;
+			UE_MT_SCOPED_READ_ACCESS(Detector);
+			AutoRTFM::AbortTransaction();
+		});
 	}
 
 	SECTION("Transact(Create, UE_MT_SCOPED_WRITE_ACCESS)")
 	{
-		AutoRTFM::ETransactionResult Result = AutoRTFM::Transact([&]
-			{
-				FRWTransactionallySafeAccessDetector Detector;
-				UE_MT_SCOPED_WRITE_ACCESS(Detector);
-			});
-		REQUIRE(AutoRTFM::ETransactionResult::Committed == Result);
+		AutoRTFM::Testing::Commit([&]
+		{
+			TestType Detector;
+			UE_MT_SCOPED_WRITE_ACCESS(Detector);
+		});
 	}
 
 	SECTION("Transact(Create, UE_MT_SCOPED_WRITE_ACCESS, Abort)")
 	{
-		AutoRTFM::ETransactionResult Result = AutoRTFM::Transact([&]
-			{
-				FRWTransactionallySafeAccessDetector Detector;
-				UE_MT_SCOPED_WRITE_ACCESS(Detector);
-				AutoRTFM::AbortTransaction();
-			});
-		REQUIRE(AutoRTFM::ETransactionResult::AbortedByRequest == Result);
+		AutoRTFM::Testing::Abort([&]
+		{
+			TestType Detector;
+			UE_MT_SCOPED_WRITE_ACCESS(Detector);
+			AutoRTFM::AbortTransaction();
+		});
 	}
 
 	SECTION("UE_MT_SCOPED_READ_ACCESS, destruct, memzero, reconstruct")
 	{
-		FRWTransactionallySafeAccessDetector Detector;
+		TestType Detector;
 		SECTION("Commit")
 		{
 			AutoRTFM::Testing::Commit([&]
@@ -918,9 +859,9 @@ TEST_CASE("MTTransactionallySafeAccessDetector")
 					// Read-lock and then unlock
 					UE_MT_SCOPED_READ_ACCESS(Detector);
 				}
-				Detector.~FRWTransactionallySafeAccessDetector();
+				Detector.~TestType();
 				memset(&Detector, 0, sizeof(Detector));
-				new (&Detector) FRWTransactionallySafeAccessDetector();
+				new (&Detector) TestType();
 			});
 		}
 		SECTION("Abort")
@@ -931,9 +872,9 @@ TEST_CASE("MTTransactionallySafeAccessDetector")
 					// Read-lock and then unlock
 					UE_MT_SCOPED_READ_ACCESS(Detector);
 				}
-				Detector.~FRWTransactionallySafeAccessDetector();
+				Detector.~TestType();
 				memset(&Detector, 0, sizeof(Detector));
-				new (&Detector) FRWTransactionallySafeAccessDetector();
+				new (&Detector) TestType();
 				AutoRTFM::AbortTransaction();
 			});
 		}
@@ -941,7 +882,7 @@ TEST_CASE("MTTransactionallySafeAccessDetector")
 
 	SECTION("UE_MT_SCOPED_WRITE_ACCESS, destruct, memzero, reconstruct")
 	{
-		FRWTransactionallySafeAccessDetector Detector;
+		TestType Detector;
 		SECTION("Commit")
 		{
 			AutoRTFM::Testing::Commit([&]
@@ -950,9 +891,9 @@ TEST_CASE("MTTransactionallySafeAccessDetector")
 					// Write-lock and then unlock
 					UE_MT_SCOPED_WRITE_ACCESS(Detector);
 				}
-				Detector.~FRWTransactionallySafeAccessDetector();
+				Detector.~TestType();
 				memset(&Detector, 0, sizeof(Detector));
-				new (&Detector) FRWTransactionallySafeAccessDetector();
+				new (&Detector) TestType();
 			});
 		}
 		SECTION("Abort")
@@ -963,9 +904,9 @@ TEST_CASE("MTTransactionallySafeAccessDetector")
 					// Write-lock and then unlock
 					UE_MT_SCOPED_WRITE_ACCESS(Detector);
 				}
-				Detector.~FRWTransactionallySafeAccessDetector();
+				Detector.~TestType();
 				memset(&Detector, 0, sizeof(Detector));
-				new (&Detector) FRWTransactionallySafeAccessDetector();
+				new (&Detector) TestType();
 				AutoRTFM::AbortTransaction();
 			});
 		}
