[AutoRTFM] Add unit tests and benchmarks for FIntervalTree

diff --git a/Engine/Source/Programs/AutoRTFMTests/Private/IntervalTreeTests.cpp b/Engine/Source/Programs/AutoRTFMTests/Private/IntervalTreeTests.cpp
new file mode 100644
--- /dev/null
+++ b/Engine/Source/Programs/AutoRTFMTests/Private/IntervalTreeTests.cpp
@@ -0,0 +1,341 @@
+// Copyright Epic Games, Inc. All Rights Reserved.
+
+#include "IntervalTree.h"
+#include "Catch2Includes.h"
+#include <cstdint>
+#include <limits>
+
+struct FInterval final
+{
+	uintptr_t Start = 0;
+	uintptr_t End = 0;
+	uintptr_t Size() const { return End - Start; }
+	bool operator < (const FInterval& Other) const { return Start < Other.Start; }
+};
+
+static constexpr FInterval MaxInterval{0, std::numeric_limits<uintptr_t>::max()};
+
+struct FTreeIntervals
+{
+	// Intervals that are not in the tree
+	std::vector<FInterval> Free;
+	// Intervals that are in the tree
+	std::vector<FInterval> Used;
+	
+	// Sorts the Free and Used intervals into ascending order, and checks that
+	// each interval in the vector has a non-zero size and does not overlap any
+	// other interval in the same vector.
+	void SortAndValidate()
+	{
+		auto SortAndValidateIntervals = [](std::vector<FInterval>& Intervals)
+		{
+			std::sort(Intervals.begin(), Intervals.end());
+			uintptr_t End = 0;
+			for (FInterval& Interval : Intervals)
+			{
+				REQUIRE(Interval.Size() > 0);
+				REQUIRE(End <= Interval.Start);
+				End = Interval.End;
+			}
+		};
+		SortAndValidateIntervals(Free);
+		SortAndValidateIntervals(Used);
+	}
+};
+
+// Populates Tree with at most MaxUsed random intervals across the entire uintptr_t interval range.
+static FTreeIntervals BuildIntervals(int Seed, size_t MaxUsed)
+{
+	std::mt19937 Rand(Seed);
+	std::vector<FInterval> Free{MaxInterval};
+	std::vector<FInterval> Used;
+
+	while (!Free.empty() && Used.size() < MaxUsed)
+	{
+		// Pop a random free interval
+		std::swap(Free[Rand() % Free.size()], Free.back());
+		const FInterval FreeInterval = Free.back();
+		Free.pop_back();
+
+		// Steal the part of the interval
+		static constexpr size_t MaxIntervalSize = 0x100000;
+		const uintptr_t Size = FreeInterval.Size() > 1 ? (1 + Rand() % std::min(MaxIntervalSize, FreeInterval.Size() - 1)) : 1;
+		uintptr_t Offset = FreeInterval.Size() > Size ? (Rand() % (FreeInterval.Size() - Size)) : 0;
+		if (Rand() % 100 < 5) // 5% of the intervals are aligned to the start or end of another interval
+		{
+			Offset = (Rand() & 1) ? 0 : FreeInterval.Size() - Size;
+		}
+		
+		REQUIRE(Offset + Size <= FreeInterval.Size());
+
+		const FInterval UsedInterval{FreeInterval.Start + Offset, FreeInterval.Start + Offset + Size};
+		
+		if (FreeInterval.Start != UsedInterval.Start)
+		{
+			Free.push_back(FInterval{FreeInterval.Start, UsedInterval.Start});
+		}
+		if (UsedInterval.End != FreeInterval.End)
+		{
+			Free.push_back(FInterval{UsedInterval.End, FreeInterval.End});
+		}
+
+		Used.push_back(UsedInterval);
+	}
+
+	return FTreeIntervals{Free, Used};
+}
+
+TEST_CASE("IntervalTree")
+{
+	// Convenience lambda to convert a uintptr_t to a void*
+	auto Ptr = [](uintptr_t Address) { return reinterpret_cast<void*>(Address); };
+
+	// Checks that the interval tree has all of the expected intervals, and that Contains() returns
+	// the correct value for various intervals.
+	auto Validate = [&](AutoRTFM::FIntervalTree& IntervalTree, std::vector<FInterval> ExpectedIntervals)
+	{
+		std::vector<FInterval> GotIntervals;
+		IntervalTree.ForEach([&](uintptr_t Start, uintptr_t End)
+		{
+			GotIntervals.push_back({Start, End});
+		});
+		REQUIRE(ExpectedIntervals.empty() == IntervalTree.IsEmpty());
+		for (size_t I = 0, N = std::min(ExpectedIntervals.size(), GotIntervals.size()); I < N; I++)
+		{
+			REQUIRE(ExpectedIntervals[I].Start == GotIntervals[I].Start);
+			REQUIRE(ExpectedIntervals[I].End == GotIntervals[I].End);
+		}
+		REQUIRE(ExpectedIntervals.size() == GotIntervals.size());
+		for (size_t IntervalIndex = 0; IntervalIndex < ExpectedIntervals.size(); IntervalIndex++)
+		{
+			const FInterval Interval = ExpectedIntervals[IntervalIndex];
+			const bool bTouchingPrev = IntervalIndex > 0 && ExpectedIntervals[IntervalIndex - 1].End == Interval.Start;
+			const bool bTouchingNext = IntervalIndex < (ExpectedIntervals.size() - 1) && ExpectedIntervals[IntervalIndex + 1].Start == Interval.End;
+			REQUIRE(true == IntervalTree.Contains(Ptr(Interval.Start), 1));
+			REQUIRE(true == IntervalTree.Contains(Ptr(Interval.Start), Interval.End - Interval.Start));
+			REQUIRE(bTouchingPrev == IntervalTree.Contains(Ptr(Interval.Start - 1), 1));
+			REQUIRE(bTouchingNext == IntervalTree.Contains(Ptr(Interval.End), 1));
+		}
+	};
+
+	AutoRTFM::FIntervalTree IntervalTree;
+	
+	Validate(IntervalTree, {});
+
+	SECTION("Insert")
+	{
+		IntervalTree.Reset();
+		Validate(IntervalTree, {});
+
+		// New intervals
+		IntervalTree.Insert(Ptr(40), 10);
+		Validate(IntervalTree, {{40, 50}});
+
+		IntervalTree.Insert(Ptr(20), 10);
+		Validate(IntervalTree, {{20, 30}, {40, 50}});
+
+		IntervalTree.Insert(Ptr(60), 10);
+		Validate(IntervalTree, {{20, 30}, {40, 50}, {60, 70}});
+
+		IntervalTree.Insert(Ptr(80), 10);
+		Validate(IntervalTree, {{20, 30}, {40, 50}, {60, 70}, {80, 90}});
+		
+		// Re-inserted intervals
+		IntervalTree.Insert(Ptr(20), 10);
+		Validate(IntervalTree, {{20, 30}, {40, 50}, {60, 70}, {80, 90}});
+
+		IntervalTree.Insert(Ptr(40), 10);
+		Validate(IntervalTree, {{20, 30}, {40, 50}, {60, 70}, {80, 90}});
+		
+		IntervalTree.Insert(Ptr(60), 10);
+		Validate(IntervalTree, {{20, 30}, {40, 50}, {60, 70}, {80, 90}});
+		
+		IntervalTree.Insert(Ptr(80), 10);
+		Validate(IntervalTree, {{20, 30}, {40, 50}, {60, 70}, {80, 90}});
+
+		// Expanding intervals without overlap
+		IntervalTree.Insert(Ptr(15), 5);
+		Validate(IntervalTree, {{15, 30}, {40, 50}, {60, 70}, {80, 90}});
+
+		IntervalTree.Insert(Ptr(30), 5);
+		Validate(IntervalTree, {{15, 35}, {40, 50}, {60, 70}, {80, 90}});
+
+		IntervalTree.Insert(Ptr(55), 5);
+		Validate(IntervalTree, {{15, 35}, {40, 50}, {55, 70}, {80, 90}});
+
+		IntervalTree.Insert(Ptr(70), 5);
+		Validate(IntervalTree, {{15, 35}, {40, 50}, {55, 75}, {80, 90}});
+
+		// Expanding intervals with overlap
+		IntervalTree.Insert(Ptr(35), 5);
+		Validate(IntervalTree, {{15, 35}, {35, 50}, {55, 75}, {80, 90}});
+
+		IntervalTree.Insert(Ptr(50), 5);
+		Validate(IntervalTree, {{15, 35}, {35, 55}, {55, 75}, {80, 90}});
+
+		IntervalTree.Insert(Ptr(75), 5);
+		Validate(IntervalTree, {{15, 35}, {35, 55}, {55, 80}, {80, 90}});
+	}
+
+	SECTION("Merge")
+	{
+		IntervalTree.Insert(Ptr(20), 10);
+		IntervalTree.Insert(Ptr(40), 10);
+		IntervalTree.Insert(Ptr(60), 10);
+		IntervalTree.Insert(Ptr(80), 10);
+		Validate(IntervalTree, {{20, 30}, {40, 50}, {60, 70}, {80, 90}});
+
+		SECTION("Intertwined")
+		{
+			AutoRTFM::FIntervalTree IntervalTreeB;
+			IntervalTreeB.Insert(Ptr(10), 5);
+			IntervalTreeB.Insert(Ptr(33), 6);
+			IntervalTreeB.Insert(Ptr(54), 2);
+			IntervalTreeB.Insert(Ptr(73), 3);
+			IntervalTreeB.Insert(Ptr(95), 2);
+			Validate(IntervalTreeB, {{10, 15}, {33, 39}, {54, 56}, {73, 76}, {95, 97}});
+			
+			IntervalTree.Merge(IntervalTreeB);
+			Validate(IntervalTree, {{10, 15}, {20, 30}, {33, 39}, {40, 50}, {54, 56}, {60, 70}, {73, 76}, {80, 90}, {95, 97}});
+		}
+
+		SECTION("Expand")
+		{
+			AutoRTFM::FIntervalTree IntervalTreeB;
+			IntervalTreeB.Insert(Ptr(15), 5);
+			IntervalTreeB.Insert(Ptr(30), 3);
+			IntervalTreeB.Insert(Ptr(37), 3);
+			IntervalTreeB.Insert(Ptr(50), 2);
+			IntervalTreeB.Insert(Ptr(55), 5);
+			IntervalTreeB.Insert(Ptr(70), 4);
+			IntervalTreeB.Insert(Ptr(78), 2);
+			IntervalTreeB.Insert(Ptr(90), 3);
+			Validate(IntervalTreeB, {{15, 20}, {30, 33}, {37, 40}, {50, 52}, {55, 60}, {70, 74}, {78, 80}, {90, 93}});
+			
+			IntervalTree.Merge(IntervalTreeB);
+			Validate(IntervalTree, {{15, 33}, {37, 52}, {55, 74}, {78, 93}});
+		}
+
+		SECTION("Overlap")
+		{
+			AutoRTFM::FIntervalTree IntervalTreeB;
+			IntervalTreeB.Insert(Ptr(30), 10);
+			IntervalTreeB.Insert(Ptr(50), 10);
+			IntervalTreeB.Insert(Ptr(70), 10);
+			Validate(IntervalTreeB, {{30, 40}, {50, 60}, {70, 80}});
+			
+			IntervalTree.Merge(IntervalTreeB);
+			Validate(IntervalTree, {{20, 30}, {30, 60}, {60, 80}, {80, 90}});
+		}
+	}
+
+	SECTION("Soak")
+	{
+		for (int Pass = 0; Pass < 32; Pass++)
+		{
+			IntervalTree.Reset();
+			Validate(IntervalTree, {});
+
+			FTreeIntervals Intervals = BuildIntervals(/* Seed */ Pass, /* MaxUsed */ 1000);
+
+			for (FInterval Interval : Intervals.Used)
+			{
+				IntervalTree.Insert(Ptr(Interval.Start), Interval.Size());
+			}
+
+			Intervals.SortAndValidate();
+
+			for (const FInterval& FreeInterval : Intervals.Free)
+			{
+				REQUIRE(!IntervalTree.Contains(Ptr(FreeInterval.Start), FreeInterval.Size()));
+				REQUIRE(!IntervalTree.Contains(Ptr(FreeInterval.Start), 1));
+				REQUIRE(!IntervalTree.Contains(Ptr(FreeInterval.End - 1), 1));
+			}
+			for (const FInterval& UsedInterval : Intervals.Used)
+			{
+				REQUIRE(IntervalTree.Contains(Ptr(UsedInterval.Start), UsedInterval.Size()));
+				REQUIRE(IntervalTree.Contains(Ptr(UsedInterval.Start), 1));
+				REQUIRE(IntervalTree.Contains(Ptr(UsedInterval.End - 1), 1));
+			}
+		}
+	}
+
+	IntervalTree.Reset();
+	Validate(IntervalTree, {});
+}
+
+TEST_CASE("IntervalTree.Benchmarks")
+{
+	// Convenience lambda to convert a uintptr_t to a void*
+	auto Ptr = [](uintptr_t Address) { return reinterpret_cast<void*>(Address); };
+
+	AutoRTFM::FIntervalTree IntervalTree;
+
+	// Populates IntervalTree with a random set of intervals.
+	// Returns the FTreeIntervals used to populate the IntervalTree.
+	auto PopulateIntervalTree = [&]() -> FTreeIntervals
+	{
+		const FTreeIntervals Intervals = BuildIntervals(/* Seed */ 1234, /* MaxUsed */ 10'000);
+		for (FInterval Interval : Intervals.Used)
+		{
+			IntervalTree.Insert(Ptr(Interval.Start), Interval.Size());
+		}
+		return Intervals;
+	};
+	
+
+	BENCHMARK_ADVANCED("Insert")(Catch::Benchmark::Chronometer Meter)
+	{
+		const std::vector<FInterval> Intervals = PopulateIntervalTree().Used;
+		Meter.measure([&]
+		{
+			for (FInterval Interval : Intervals)
+			{
+				IntervalTree.Insert(Ptr(Interval.Start), Interval.Size());
+			}
+		});
+	};
+
+	// Populates IntervalTree with a random set of intervals.
+	// If bUsed is true, returns a random set of addresses in those intervals,
+	// otherwise a random set of addresses outside of those intervals.
+	auto PopulateIntervalTreeAndPickAddresses = [&](bool bUsed) -> std::vector<void*>
+	{
+		std::mt19937 Rand(4242);
+		
+		const std::vector<FInterval> Intervals = bUsed ? PopulateIntervalTree().Used : PopulateIntervalTree().Free;
+
+		constexpr size_t NumAddresses = 10'000;
+		std::vector<void*> Addresses(NumAddresses);
+		for (size_t I = 0; I < NumAddresses; I++)
+		{
+			FInterval Interval = Intervals[Rand() % Intervals.size()];
+			Addresses[I] = Ptr(Interval.Start + (Rand() % Interval.Size()));
+		}
+		return Addresses;
+	};
+
+	BENCHMARK_ADVANCED("Contains.Used")(Catch::Benchmark::Chronometer Meter)
+	{
+		std::vector<void*> Addresses = PopulateIntervalTreeAndPickAddresses(/* bUsed */ true);
+		Meter.measure([&]
+		{
+			for (void* Address : Addresses)
+			{
+				IntervalTree.Contains(Address, 4);
+			}
+		});
+	};
+
+	BENCHMARK_ADVANCED("Contains.Free")(Catch::Benchmark::Chronometer Meter)
+	{
+		std::vector<void*> Addresses = PopulateIntervalTreeAndPickAddresses(/* bUsed */ false);
+		Meter.measure([&]
+		{
+			for (void* Address : Addresses)
+			{
+				IntervalTree.Contains(Address, 4);
+			}
+		});
+	};
+}
diff --git a/Engine/Source/Runtime/AutoRTFM/Private/IntervalTree.h b/Engine/Source/Runtime/AutoRTFM/Private/IntervalTree.h
--- a/Engine/Source/Runtime/AutoRTFM/Private/IntervalTree.h
+++ b/Engine/Source/Runtime/AutoRTFM/Private/IntervalTree.h
@@ -105,6 +105,17 @@ struct FIntervalTree final
         } while (!ToProcess.IsEmpty());
     }
 
+	// Calls Callback with each interval in the tree, in order of address.
+	// Callback must have the signature: Callback(uintptr_t Start, uintptr_t End)
+	template<typename CALLBACK>
+	void ForEach(CALLBACK&& Callback) const
+	{
+		if (Root != IntervalTreeNodeIndexNone)
+		{
+			ForEach(Root, Callback);
+		}
+	}
+
 private:
 	static constexpr size_t InlineArraySize = 8;
 	template<typename T> using StackType = TStack<T, InlineArraySize, EContainerValidation::Disabled>;
@@ -138,6 +149,20 @@ private:
 
 	static constexpr bool bExtraDebugging = false;
 
+	template<typename CALLBACK>
+	void ForEach(FIntervalTreeNodeIndex Index, CALLBACK&& Callback) const
+	{
+		if (FIntervalTreeNodeIndex Left = NodeLefts[Index]; Left != IntervalTreeNodeIndexNone)
+		{
+			ForEach(Left, Callback);
+		}
+		Callback(NodeRanges[Index].Start, NodeRanges[Index].End);
+		if (FIntervalTreeNodeIndex Right = NodeRights[Index]; Right != IntervalTreeNodeIndexNone)
+		{
+			ForEach(Right, Callback);
+		}
+	}
+
 	UE_AUTORTFM_FORCEINLINE FIntervalTreeNodeIndex AddNode(FRange&& NewRange)
 	{
 		const FIntervalTreeNodeIndex Index = static_cast<FIntervalTreeNodeIndex>(NodeRanges.Num());
