Add #include for FStringOutputDevice now that it is separated into its own header.

diff --git a/Engine/Source/Programs/AutoRTFMTests/Private/UECoreTests.cpp b/Engine/Source/Programs/AutoRTFMTests/Private/UECoreTests.cpp
--- a/Engine/Source/Programs/AutoRTFMTests/Private/UECoreTests.cpp
+++ b/Engine/Source/Programs/AutoRTFMTests/Private/UECoreTests.cpp
@@ -36,6 +36,7 @@
 #include "Misc/FileHelper.h"
 #include "Misc/PackageName.h"
 #include "Misc/ScopeRWLock.h"
+#include "Misc/StringOutputDevice.h"
 #include "StructUtils/PropertyBag.h"
 
 #include <thread>
