[AutoRTFM] Add unit tests for AutoRTFM::OnComplete.

This didn't expose any bugs, but it did cause me to re-evaluate some transaction state names. Previously we had failure status codes for "attempted to transact in an OnCommit" or "in an OnAbort." The status codes are now renamed to "attempted to transact during commit" or "during abort." This subtle difference accounts for the fact that you could attempt to open a transact during the OnComplete phase, and this will error out with "attempted to transact during commit" or "during abort" depending on the final outcome of the transaction that attempted to transact within its OnComplete.

Also, one unit test was difficult to implement because our retry mechanism in the unit test harness led to unexpected status codes (an abort in a test which didn't have any aborts in it). This needed debugging, but ultimately just needed some extra checks in the unit test to guard against an unexpected abort status.

diff --git a/Engine/Source/Programs/AutoRTFMTests/Private/API.cpp b/Engine/Source/Programs/AutoRTFMTests/Private/API.cpp
--- a/Engine/Source/Programs/AutoRTFMTests/Private/API.cpp
+++ b/Engine/Source/Programs/AutoRTFMTests/Private/API.cpp
@@ -377,8 +377,9 @@ TEST_CASE("API.IsTransactional")
     bool InOpenNest = false;
 	bool InAbort = true;
 	bool InCommit = true;
+	bool InComplete = true;
 
-    AutoRTFM::Commit([&]
+    AutoRTFM::Testing::Commit([&]
     {
         InTransaction = AutoRTFM::IsTransactional();
 
@@ -387,26 +388,32 @@ TEST_CASE("API.IsTransactional")
             InOpenNest = AutoRTFM::IsTransactional();
         });
 
-		AutoRTFM::Transact([&]
+		AutoRTFM::Testing::Abort([&]
+		{
+			AutoRTFM::OnAbort([&]
 			{
-				AutoRTFM::OnAbort([&]
-					{
-						InAbort = AutoRTFM::IsTransactional();
-					});
-
-				AutoRTFM::AbortTransaction();
+				InAbort = AutoRTFM::IsTransactional();
 			});
 
+			AutoRTFM::AbortTransaction();
+		});
+
 		AutoRTFM::OnCommit([&]
-			{
-				InCommit = AutoRTFM::IsTransactional();
-			});
+		{
+			InCommit = AutoRTFM::IsTransactional();
+		});
+
+		AutoRTFM::OnComplete([&]
+		{
+			InComplete = AutoRTFM::IsTransactional();
+		});
     });
 
     REQUIRE(true == InTransaction);
     REQUIRE(true == InOpenNest);
 	REQUIRE(false == InAbort);
 	REQUIRE(false == InCommit);
+	REQUIRE(false == InComplete);
 }
 
 TEST_CASE("API.IsClosed")
@@ -419,25 +426,31 @@ TEST_CASE("API.IsClosed")
     bool InClosedNestInOpenNest = false;
 	bool InAbort = true;
 	bool InCommit = true;
+	bool InComplete = true;
 
-    AutoRTFM::Commit([&]
+    AutoRTFM::Testing::Commit([&]
     {
         InTransaction = AutoRTFM::IsClosed();
 
-		AutoRTFM::Transact([&]
+		AutoRTFM::Testing::Abort([&]
+		{
+			AutoRTFM::OnAbort([&]
 			{
-				AutoRTFM::OnAbort([&]
-					{
-						InAbort = AutoRTFM::IsClosed();
-					});
-
-				AutoRTFM::AbortTransaction();
+				InAbort = AutoRTFM::IsClosed();
 			});
 
+			AutoRTFM::AbortTransaction();
+		});
+
 		AutoRTFM::OnCommit([&]
-			{
-				InCommit = AutoRTFM::IsClosed();
-			});
+		{
+			InCommit = AutoRTFM::IsClosed();
+		});
+
+		AutoRTFM::OnComplete([&]
+		{
+			InComplete = AutoRTFM::IsClosed();
+		});
 
         AutoRTFM::Open([&]
         {
@@ -455,6 +468,7 @@ TEST_CASE("API.IsClosed")
     REQUIRE(true == InClosedNestInOpenNest);
 	REQUIRE(false == InAbort);
 	REQUIRE(false == InCommit);
+	REQUIRE(false == InComplete);
 }
 
 TEST_CASE("API.IsCommittingOrAborting")
@@ -467,42 +481,49 @@ TEST_CASE("API.IsCommittingOrAborting")
 	bool InClosedNestInOpenNest = true;
 	bool InAbort = false;
 	bool InCommit = false;
+	bool InComplete = false;
 
-	AutoRTFM::Commit([&]
+	AutoRTFM::Testing::Commit([&]
+	{
+		InTransaction = AutoRTFM::IsCommittingOrAborting();
+
+		AutoRTFM::Testing::Abort([&]
 		{
-			InTransaction = AutoRTFM::IsCommittingOrAborting();
+			AutoRTFM::OnAbort([&]
+			{
+				InAbort = AutoRTFM::IsCommittingOrAborting();
+			});
 
-			AutoRTFM::Transact([&]
-				{
-					AutoRTFM::OnAbort([&]
-						{
-							InAbort = AutoRTFM::IsCommittingOrAborting();
-						});
+			AutoRTFM::AbortTransaction();
+		});
 
-					AutoRTFM::AbortTransaction();
-				});
+		AutoRTFM::OnCommit([&]
+		{
+			InCommit = AutoRTFM::IsCommittingOrAborting();
+		});
 
-			AutoRTFM::OnCommit([&]
-				{
-					InCommit = AutoRTFM::IsCommittingOrAborting();
-				});
+		AutoRTFM::OnComplete([&]
+		{
+			InComplete = AutoRTFM::IsCommittingOrAborting();
+		});
 
-			AutoRTFM::Open([&]
-				{
-					InOpenNest = AutoRTFM::IsCommittingOrAborting();
+		AutoRTFM::Open([&]
+		{
+			InOpenNest = AutoRTFM::IsCommittingOrAborting();
 
-					REQUIRE(AutoRTFM::EContextStatus::OnTrack == AutoRTFM::Close([&]
-						{
-							InClosedNestInOpenNest = AutoRTFM::IsCommittingOrAborting();
-						}));
-				});
+			REQUIRE(AutoRTFM::EContextStatus::OnTrack == AutoRTFM::Close([&]
+			{
+				InClosedNestInOpenNest = AutoRTFM::IsCommittingOrAborting();
+			}));
 		});
+	});
 
 	REQUIRE(false == InTransaction);
 	REQUIRE(false == InOpenNest);
 	REQUIRE(false == InClosedNestInOpenNest);
 	REQUIRE(true == InAbort);
 	REQUIRE(true == InCommit);
+	REQUIRE(true == InComplete);
 }
 
 TEST_CASE("API.Transact")
@@ -748,6 +769,118 @@ TEST_CASE("API.OnCommit")
     REQUIRE(true == InnerOpenNest);
 }
 
+TEST_CASE("API.OnComplete")
+{
+	bool bTriggered = false;
+
+	SECTION("Commit(OnComplete)")
+	{
+		AutoRTFM::Testing::Commit([&]
+		{
+			// We reset `bTriggered` at the top of each transaction, because our retry testing can 
+			// cause us to rerun the transaction after `bTriggered` has been set via OnComplete.
+			bTriggered = false;
+			AutoRTFM::OnComplete([&]
+			{
+				bTriggered = true;
+			});
+			REQUIRE(!bTriggered);
+		});
+		REQUIRE(bTriggered);
+	}
+
+	SECTION("Abort(OnComplete)")
+	{
+		AutoRTFM::Testing::Abort([&]
+		{
+			bTriggered = false;
+			AutoRTFM::OnComplete([&]
+			{
+				bTriggered = true;
+			});
+			REQUIRE(!bTriggered);
+			AutoRTFM::AbortTransaction();
+		});
+		REQUIRE(bTriggered);
+	}
+
+	SECTION("Commit(Commit(OnComplete))")
+	{
+		AutoRTFM::Testing::Commit([&]
+		{
+			bTriggered = false;
+			AutoRTFM::Testing::Commit([&]
+			{
+				AutoRTFM::OnComplete([&]
+				{
+					bTriggered = true;
+				});
+				REQUIRE(!bTriggered);
+			});
+			REQUIRE(!bTriggered);
+		});
+		REQUIRE(bTriggered);
+	}
+
+	SECTION("Abort(Abort(OnComplete))")
+	{
+		AutoRTFM::Testing::Abort([&]
+		{
+			bTriggered = false;
+			AutoRTFM::Testing::Abort([&]
+			{
+				AutoRTFM::OnComplete([&]
+				{
+					bTriggered = true;
+				});
+				REQUIRE(!bTriggered);
+				AutoRTFM::AbortTransaction();
+			});
+			REQUIRE(!bTriggered);
+			AutoRTFM::AbortTransaction();
+		});
+		REQUIRE(bTriggered);
+	}
+
+	SECTION("Commit(Abort(OnComplete))")
+	{
+		AutoRTFM::Testing::Commit([&]
+		{
+			bTriggered = false;
+			AutoRTFM::Testing::Abort([&]
+			{
+				AutoRTFM::OnComplete([&]
+				{
+					bTriggered = true;
+				});
+				REQUIRE(!bTriggered);
+				AutoRTFM::AbortTransaction();
+			});
+			REQUIRE(!bTriggered);
+		});
+		REQUIRE(bTriggered);
+	}
+
+	SECTION("Abort(Commit(OnComplete))")
+	{
+		AutoRTFM::Testing::Abort([&]
+		{
+			bTriggered = false;
+			AutoRTFM::Testing::Commit([&]
+			{
+				AutoRTFM::OnComplete([&]
+				{
+					bTriggered = true;
+				});
+				REQUIRE(!bTriggered);
+			});
+			REQUIRE(!bTriggered);
+			AutoRTFM::AbortTransaction();
+		});
+		REQUIRE(bTriggered);
+	}
+}
+
 TEST_CASE("API.OnCommit_MutableCapture")
 {
 	FString Message = "Hello";
@@ -762,6 +895,20 @@ TEST_CASE("API.OnCommit_MutableCapture")
 	});
 }
 
+TEST_CASE("API.OnComplete_MutableCapture")
+{
+	FString Message = "Hello";
+
+	AutoRTFM::Transact([&]
+	{
+		AutoRTFM::OnComplete([MessageCopy = Message]() mutable
+		{
+			MessageCopy += " World!";
+			REQUIRE(MessageCopy == "Hello World!");
+		});
+	});
+}
+
 TEST_CASE("API.OnCommitMacro_NoAbort")
 {
 	int Value = 123;
diff --git a/Engine/Source/Programs/AutoRTFMTests/Private/Abort.cpp b/Engine/Source/Programs/AutoRTFMTests/Private/Abort.cpp
--- a/Engine/Source/Programs/AutoRTFMTests/Private/Abort.cpp
+++ b/Engine/Source/Programs/AutoRTFMTests/Private/Abort.cpp
@@ -98,7 +98,7 @@ TEST_CASE("Abort.TransactionInOnCommit")
 				bDidSomething = true;
 			});
 
-			REQUIRE(AutoRTFM::ETransactionResult::AbortedByTransactInOnCommit == Result);
+			REQUIRE(AutoRTFM::ETransactionResult::AbortedByTransactDuringCommit == Result);
 			REQUIRE(false == bDidSomething);
 		});
 	});
@@ -118,7 +118,7 @@ TEST_CASE("Abort.TransactionInOnAbort")
 				bDidSomething = true;
 			});
 
-			REQUIRE(AutoRTFM::ETransactionResult::AbortedByTransactInOnAbort == InnerResult);
+			REQUIRE(AutoRTFM::ETransactionResult::AbortedByTransactDuringAbort == InnerResult);
 			REQUIRE(false == bDidSomething);
 		});
 
@@ -128,6 +128,57 @@ TEST_CASE("Abort.TransactionInOnAbort")
 	REQUIRE(AutoRTFM::ETransactionResult::AbortedByRequest == Result);
 }
 
+TEST_CASE("Abort.TransactionInAbortedOnComplete")
+{
+	AutoRTFM::Testing::Abort([&]
+	{
+		AutoRTFM::OnComplete([&]
+		{
+			bool bDidSomething = false;
+
+			AutoRTFM::ETransactionResult InnerResult = AutoRTFM::Transact([&]
+			{
+				bDidSomething = true;
+			});
+
+			REQUIRE(InnerResult == AutoRTFM::ETransactionResult::AbortedByTransactDuringAbort);
+			REQUIRE(!bDidSomething);
+		});
+
+		AutoRTFM::AbortTransaction();
+	});
+}
+
+TEST_CASE("Abort.TransactionInCommittedOnComplete")
+{
+	bool bGotAbortedDuringCommit = false;
+
+	AutoRTFM::Testing::Commit([&]
+	{
+		AutoRTFM::OnComplete([&]
+		{
+			bool bDidSomething = false;
+
+			AutoRTFM::ETransactionResult InnerResult = AutoRTFM::Transact([&]
+			{
+				bDidSomething = true;
+			});
+
+			// The test harness will abort this transaction and retry it when `ShouldRetryNonNestedTransactions` 
+			// is set, which unfortunately muddles the test result here since we'd normally expect "transact during
+			// commit" here, not "transact during abort."
+			REQUIRE((InnerResult == AutoRTFM::ETransactionResult::AbortedByTransactDuringCommit ||
+			         InnerResult == AutoRTFM::ETransactionResult::AbortedByTransactDuringAbort));
+
+			bGotAbortedDuringCommit |= (InnerResult == AutoRTFM::ETransactionResult::AbortedByTransactDuringCommit);
+			REQUIRE(!bDidSomething);
+		});
+	});
+
+	// The test only passes if we actually got aborted during commit, so verify that this really happened.
+	REQUIRE(bGotAbortedDuringCommit);
+}
+
 TEST_CASE("Abort.AbortInInnerTransaction")
 {
 	int Value = 1;
diff --git a/Engine/Source/Runtime/AutoRTFM/Private/Context.cpp b/Engine/Source/Runtime/AutoRTFM/Private/Context.cpp
--- a/Engine/Source/Runtime/AutoRTFM/Private/Context.cpp
+++ b/Engine/Source/Runtime/AutoRTFM/Private/Context.cpp
@@ -414,12 +414,12 @@ ETransactionResult FContext::Transact(void (*UninstrumentedFunction)(void*), voi
 {
     if (AUTORTFM_UNLIKELY(EContextStatus::Committing == Status))
     {
-    	return ETransactionResult::AbortedByTransactInOnCommit;
+    	return ETransactionResult::AbortedByTransactDuringCommit;
     }
 
     if (AUTORTFM_UNLIKELY(IsAborting()))
     {
-    	return ETransactionResult::AbortedByTransactInOnAbort;
+    	return ETransactionResult::AbortedByTransactDuringAbort;
     }
     
     AUTORTFM_ASSERT(Status == EContextStatus::Idle || Status == EContextStatus::OnTrack);
diff --git a/Engine/Source/Runtime/AutoRTFM/Public/AutoRTFM.h b/Engine/Source/Runtime/AutoRTFM/Public/AutoRTFM.h
--- a/Engine/Source/Runtime/AutoRTFM/Public/AutoRTFM.h
+++ b/Engine/Source/Runtime/AutoRTFM/Public/AutoRTFM.h
@@ -45,8 +45,8 @@ typedef enum
     autortfm_aborted_by_request = 0,
     autortfm_aborted_by_language,
     autortfm_committed,
-	autortfm_aborted_by_transact_in_on_commit,
-	autortfm_aborted_by_transact_in_on_abort,
+	autortfm_aborted_by_transact_during_commit,
+	autortfm_aborted_by_transact_during_abort,
 	autortfm_aborted_by_cascade,
 } autortfm_result;
 
@@ -452,11 +452,13 @@ enum class ETransactionResult
 	// cannot be undone later if the parent transaction is aborted for any reason.
     Committed = autortfm_committed,
 
-	// The transaction aborted because in a call to OnCommit, a new transaction nest was attempted which is not allowed.
-	AbortedByTransactInOnCommit = autortfm_aborted_by_transact_in_on_commit,
+	// The transaction aborted because a new transaction nest was attempted (via OnCommit or OnComplete) while the 
+	// current transaction was being committed.
+	AbortedByTransactDuringCommit = autortfm_aborted_by_transact_during_commit,
 
-	// The transaction aborted because in a call to OnAbort, a new transaction nest was attempted which is not allowed.
-	AbortedByTransactInOnAbort = autortfm_aborted_by_transact_in_on_abort,
+	// The transaction aborted because a new transaction nest was attempted (via OnAbort or OnComplete) while the 
+	// current transaction was being aborted.
+	AbortedByTransactDuringAbort = autortfm_aborted_by_transact_during_abort,
 
 	// The transaction aborted because of an explicit call to CascadingAbortTransaction.
 	AbortedByCascade = autortfm_aborted_by_cascade
