[AutoRTFM] Initialize the garbage collector before calling CollectGarbage.

This allows us to run an individual test from AutoRTFMTests without crashing.

diff --git a/Engine/Source/Programs/AutoRTFMTests/Private/AutoRTFMTests.cpp b/Engine/Source/Programs/AutoRTFMTests/Private/AutoRTFMTests.cpp
--- a/Engine/Source/Programs/AutoRTFMTests/Private/AutoRTFMTests.cpp
+++ b/Engine/Source/Programs/AutoRTFMTests/Private/AutoRTFMTests.cpp
@@ -99,6 +99,10 @@ int RunTests(int ArgC, const char* ArgV[])
 		AutoRTFM::ForTheRuntime::SetRetryTransaction(AutoRTFM::ForTheRuntime::EAutoRTFMRetryTransactionState::RetryNonNested);
 	}
 
+	// Initialize the garbage collector. (Otherwise, if no active test uses garbage collection, 
+	// the call to `CollectGarbage` below might fail.)
+	FGCObject::StaticInit();
+
 	// Enable AutoRTFM.
 	AutoRTFM::ForTheRuntime::SetAutoRTFMRuntime(AutoRTFM::ForTheRuntime::EAutoRTFMEnabledState::AutoRTFM_EnabledByDefault);
 
