[AutoRTFM] Add benchmark for transactional memory writes.

This benchmark simply allocates the amount of time it takes to fill a buffer with simple writes--1, 2, 4, 8, 32, and 100 bytes. We measure both sequential and non-sequential writes, as non-sequential writes are much more expensive in the write log.

We wrap each subsequent pass with a new transaction to ensure that the write log is actually populated, as repeated writes to the same location do not need to be logged.

diff --git a/Engine/Source/Programs/AutoRTFMTests/Private/Benchmarks.cpp b/Engine/Source/Programs/AutoRTFMTests/Private/Benchmarks.cpp
--- a/Engine/Source/Programs/AutoRTFMTests/Private/Benchmarks.cpp
+++ b/Engine/Source/Programs/AutoRTFMTests/Private/Benchmarks.cpp
@@ -750,3 +750,68 @@ TEST_CASE("Benchmarks.TopLevelTransaction")
 		REQUIRE(Counter == 0);
 	};
 }
+
+namespace
+{
+	template <int Size>
+	struct RawBytes 
+	{
+		std::array<uint8, Size> Data{};
+		RawBytes<Size>& operator++()
+		{
+			for (uint8& N : Data)
+			{
+				++N;
+			}
+			return *this;
+		}
+	};
+}
+
+TEST_CASE("Benchmarks.Write")
+{
+	auto Benchmark = [&](Catch::Benchmark::Chronometer Meter, auto Value, int Stride)
+	{
+		using T = decltype(Value);
+		constexpr size_t NumElements = 50000;
+		std::vector<T> MyArray(NumElements * Stride);
+
+		Meter.measure([&]
+		{
+			AutoRTFM::Transact([&]
+			{
+				T* Ptr = &MyArray.front();
+				for (int Index = 0; Index < NumElements; ++Index)
+				{
+					*Ptr = Value;
+					Ptr += Stride;
+				}
+			});
+			++Value;
+		});
+	};
+
+#define BENCH(T, Stride) \
+	BENCHMARK_ADVANCED("Write: " #T " (Stride " #Stride ")")(Catch::Benchmark::Chronometer Meter) \
+	{ \
+		Benchmark(Meter, T{}, Stride); \
+	};
+
+	// Foldable
+	BENCH(char, 1);
+	BENCH(short, 1);
+	BENCH(int32, 1);
+	BENCH(int64, 1);
+	BENCH(RawBytes<32>, 1);
+	BENCH(RawBytes<100>, 1);
+
+	// Non-foldable
+	BENCH(char, 2);
+	BENCH(short, 2);
+	BENCH(int32, 2);
+	BENCH(int64, 2);
+	BENCH(RawBytes<32>, 2);
+	BENCH(RawBytes<100>, 2);
+
+#undef BENCH
+}
