[AutoRTFM] Add disabled test case for lock-in-closed then unlock-in-open.

This currently triggers a validation failure.
(Note: we don't support Catch2's SKIP macro because it requires exceptions.)

diff --git a/Engine/Source/Programs/AutoRTFMTests/Private/MutexTests.cpp b/Engine/Source/Programs/AutoRTFMTests/Private/MutexTests.cpp
--- a/Engine/Source/Programs/AutoRTFMTests/Private/MutexTests.cpp
+++ b/Engine/Source/Programs/AutoRTFMTests/Private/MutexTests.cpp
@@ -834,3 +834,21 @@ TEST_CASE("TransactionallySafeMutex.Locked Mutex In Destructed Object")
 		REQUIRE(!Mine);
 	});
 }
+
+TEST_CASE("TransactionallySafeMutex.Closed Lock Then Open Unlock")
+{
+	// TODO(SOL-8282): Mutex.Unlock() causes the memory validator to fail.
+
+/*
+	UE::FTransactionallySafeMutex Mutex;
+	AutoRTFM::Testing::Commit([&]
+	{
+		Mutex.Lock();
+
+		AutoRTFM::Open([&]
+		{
+			Mutex.Unlock();
+		});
+	});
+*/
+} 
