[AutoRTFM] Fix failure when constructing a std::string from the closed.

This works locally on Windows without issue, but fails on Horde: https://horde.devtools.epicgames.com/job/685ebba4c10a8cd6dfe24e95?step=700a&issue=958491

diff --git a/Engine/Source/Programs/AutoRTFMTests/Private/QueueTests.cpp b/Engine/Source/Programs/AutoRTFMTests/Private/QueueTests.cpp
--- a/Engine/Source/Programs/AutoRTFMTests/Private/QueueTests.cpp
+++ b/Engine/Source/Programs/AutoRTFMTests/Private/QueueTests.cpp
@@ -4,7 +4,6 @@
 #include "Catch2Includes.h"
 #include "Containers/TransactionallySafeSpscQueue.h"
 
-#include <string>
 #include <thread>
 
 // Test loosely based on code from Containers/ConcurrentQueuesTest.cpp.
@@ -73,9 +72,9 @@ TEST_CASE("TransactionallySafeSpscQueue")
 {
 	// Non-transactional, trivial and complex types
 	TestSpscQueueCorrectness<TTransactionallySafeSpscQueue<int>>(10000, /*bTransact=*/false);
-	TestSpscQueueCorrectness<TTransactionallySafeSpscQueue<std::string>>(10000, /*bTransact=*/false);
+	TestSpscQueueCorrectness<TTransactionallySafeSpscQueue<FString>>(10000, /*bTransact=*/false);
 
 	// Transactional, trivial and complex types
 	TestSpscQueueCorrectness<TTransactionallySafeSpscQueue<int>>(10000, /*bTransact=*/true);
-	TestSpscQueueCorrectness<TTransactionallySafeSpscQueue<std::string>>(10000, /*bTransact=*/true);
+	TestSpscQueueCorrectness<TTransactionallySafeSpscQueue<FString>>(10000, /*bTransact=*/true);
 }
