[AutoRTFM] Improve flexibility of TBlockAllocator.

Two new pieces of functionality here:
- Callers can now request that the initial allocation is a different size from the fixed-size buffer.
- Callers can now request "allocate at least," where the block allocator will give as large of a block as it can.

Because the allocation size is flexible here, allocations now return a span of {std::byte*, size_t}. Previously we just returned a void* but there didn't seem to be a good reason not to return byte* (and some callers were just casting it anyway).

I am experimenting with using the block allocator logic within the write log, and these changes would be useful for that.

diff --git a/Engine/Source/Programs/AutoRTFMTests/Private/BlockAllocatorTests.cpp b/Engine/Source/Programs/AutoRTFMTests/Private/BlockAllocatorTests.cpp
--- a/Engine/Source/Programs/AutoRTFMTests/Private/BlockAllocatorTests.cpp
+++ b/Engine/Source/Programs/AutoRTFMTests/Private/BlockAllocatorTests.cpp
@@ -9,7 +9,8 @@ TEST_CASE("BlockAllocator")
 	static constexpr size_t AllocatorBlockSize = 128;
 	static constexpr size_t AllocatorAlignment = 16;
 	static constexpr size_t AllocatorGrowthPercentage = 100; // Fixed size
-	AutoRTFM::TBlockAllocator<AllocatorBlockSize, AllocatorAlignment, AllocatorGrowthPercentage> BlockAllocator;
+	using FBlockAllocator = AutoRTFM::TBlockAllocator<AllocatorBlockSize, AllocatorAlignment, AllocatorGrowthPercentage>;
+	FBlockAllocator BlockAllocator;
 
 	struct FAllocationInfo
 	{
@@ -30,12 +31,14 @@ TEST_CASE("BlockAllocator")
 	// alignment.
 	auto Allocate = [&](size_t Size, size_t Alignment)
 	{
-		std::byte* Ptr = static_cast<std::byte*>(BlockAllocator.Allocate(Size, Alignment));
+		FBlockAllocator::FAllocatedSpan Span = BlockAllocator.Allocate(Size, Alignment);
+		REQUIRE(Span.Size == Size);
+		std::byte* Ptr = Span.Data;
 		REQUIRE(reinterpret_cast<uintptr_t>(Ptr) % Alignment == 0);
 		std::byte Value = static_cast<std::byte>(Allocations.size());
 		memset(Ptr, static_cast<int>(Value), Size);
 		Allocations.push_back(FAllocationInfo{Ptr, Size, Value});
-		return static_cast<std::byte*>(Ptr);
+		return Ptr;
 	};
 
 	SECTION("Contiguous")
@@ -109,3 +112,74 @@ TEST_CASE("BlockAllocator")
 		}
 	}
 }
+
+TEST_CASE("BlockAllocator.AllocateAtLeast.Fixed")
+{
+	static constexpr size_t AllocatorBlockSize = 128;
+	static constexpr size_t AllocatorAlignment = 1;
+	static constexpr size_t AllocatorGrowthPercentage = 100; // Fixed size
+	using FBlockAllocator = AutoRTFM::TBlockAllocator<AllocatorBlockSize, AllocatorAlignment, AllocatorGrowthPercentage>;
+	FBlockAllocator Alloc;
+
+	// Calling AllocateAtLeast on a fixed-growth block allocator will yield continual fixed-size blocks.
+	for (int Index = 1; Index <= 10; ++Index)
+	{
+		const FBlockAllocator::FAllocatedSpan Span = Alloc.AllocateAtLeast(1024, AllocatorAlignment);
+		REQUIRE(Span.Data);
+		REQUIRE(Span.Size == 1024);
+	}
+}
+
+TEST_CASE("BlockAllocator.AllocateAtLeast.Increasing")
+{
+	static constexpr size_t AllocatorBlockSize = 128;
+	static constexpr size_t AllocatorAlignment = 1;
+	static constexpr size_t AllocatorGrowthPercentage = 300; // 3x growth
+	using FBlockAllocator = AutoRTFM::TBlockAllocator<AllocatorBlockSize, AllocatorAlignment, AllocatorGrowthPercentage>;
+	FBlockAllocator Alloc;
+
+	// Calling AllocateAtLeast on a 3x-growth block allocator will yield increasingly larger size blocks.
+	size_t Expected = 128;
+	for (int Index = 1; Index <= 10; ++Index)
+	{
+		const FBlockAllocator::FAllocatedSpan Span = Alloc.AllocateAtLeast(128, AllocatorAlignment);
+		REQUIRE(Span.Data);
+		REQUIRE(Span.Size == Expected);
+		Expected *= AllocatorGrowthPercentage / 100;
+	}
+}
+
+TEST_CASE("BlockAllocator.FirstAllocatedBlockSize")
+{
+	static constexpr size_t AllocatorBlockSize = 42;
+	static constexpr size_t AllocatorAlignment = 1;
+	static constexpr size_t AllocatorGrowthPercentage = 200;
+	static constexpr size_t FirstAllocatedBlockSize = 4321;
+	using FBlockAllocator = AutoRTFM::TBlockAllocator<AllocatorBlockSize, AllocatorAlignment, AllocatorGrowthPercentage>;
+	FBlockAllocator Alloc{FirstAllocatedBlockSize};
+
+	// First, we should get the internal block.
+	{
+		const FBlockAllocator::FAllocatedSpan Span = Alloc.AllocateAtLeast(1, AllocatorAlignment);
+		REQUIRE(Span.Data);
+		REQUIRE(Span.Size == AllocatorBlockSize);
+	}
+	// Next, we should get an allocated block at the requested first-allocation size.
+	{
+		const FBlockAllocator::FAllocatedSpan Span = Alloc.AllocateAtLeast(1, AllocatorAlignment);
+		REQUIRE(Span.Data);
+		REQUIRE(Span.Size == FirstAllocatedBlockSize);
+	}
+	// Next, we should get allocated blocks that grow at the requested rate: first 2x...
+	{
+		const FBlockAllocator::FAllocatedSpan Span = Alloc.AllocateAtLeast(1, AllocatorAlignment);
+		REQUIRE(Span.Data);
+		REQUIRE(Span.Size == FirstAllocatedBlockSize * 2);
+	}
+	// ... then 4x.
+	{
+		const FBlockAllocator::FAllocatedSpan Span = Alloc.AllocateAtLeast(1, AllocatorAlignment);
+		REQUIRE(Span.Data);
+		REQUIRE(Span.Size == FirstAllocatedBlockSize * 4);
+	}
+}
diff --git a/Engine/Source/Runtime/AutoRTFM/Private/BlockAllocator.h b/Engine/Source/Runtime/AutoRTFM/Private/BlockAllocator.h
--- a/Engine/Source/Runtime/AutoRTFM/Private/BlockAllocator.h
+++ b/Engine/Source/Runtime/AutoRTFM/Private/BlockAllocator.h
@@ -26,7 +26,10 @@ class TBlockAllocator
 {
 public:
 	// Constructor
-	TBlockAllocator() = default;
+	TBlockAllocator(size_t FirstAllocatedBlockSize = DefaultFirstBlockSize)
+		: NextBlockSize(FirstAllocatedBlockSize)
+	{
+	}
 
 	// Destructor - frees all the memory allocated by the allocator.
 	~TBlockAllocator()
@@ -49,11 +52,17 @@ public:
 
 	// Allocates memory from the block allocator.
 	// Alignment must be a power of two and no larger than DataAlignment.
-	inline void* Allocate(size_t Size, size_t Alignment)
+	struct FAllocatedSpan
+	{
+		std::byte* Data;
+		size_t Size;
+	};
+
+	inline FAllocatedSpan Allocate(size_t Size, size_t Alignment)
 	{
 		AUTORTFM_ASSERT(Alignment <= DataAlignment);
 
-		if (void* Allocation = Tail->TryAllocate(Size, Alignment))
+		if (FAllocatedSpan Allocation = Tail->TryAllocate(Size, Alignment); Allocation.Size == Size)
 		{
 			return Allocation;
 		}
@@ -63,11 +72,23 @@ public:
 		FBlockHeader* NewBlock = FBlockHeader::New(Tail, NewBlockSize);
 		Tail = NewBlock;
 
-		void* Allocation = NewBlock->TryAllocate(Size, Alignment);
-		AUTORTFM_ASSERT(Allocation);
+		FAllocatedSpan Allocation = NewBlock->TryAllocate(Size, Alignment);
+		AUTORTFM_ASSERT(Allocation.Size == Size);
 		return Allocation;
 	}
 
+	// Allocates memory from the block allocator, yielding as large of a block as possible.
+	// If a partial block exists which holds at least `Size`, yields the remainder of that block.
+	// Otherwise, creates a new block and yields the entire thing.
+	inline FAllocatedSpan AllocateAtLeast(size_t Size, size_t Alignment)
+	{
+		AUTORTFM_ASSERT(Alignment <= DataAlignment);
+
+		const size_t Remaining = Tail->RemainingAlignedBytes(Alignment);
+		const size_t Amount = (Remaining >= Size) ? Remaining : std::max(NextBlockSize, Size);
+		return Allocate(Amount, Alignment);
+	}
+
 	// Constructs and returns a new T into the memory returned by calling
 	// Allocate(sizeof(T), alignof(T)).
 	template<typename T, typename... ArgTypes> T* New(ArgTypes&&... Args)
@@ -102,22 +123,30 @@ private:
 
 		// Attempts to sub-allocate with out of this block. Returns a pointer to
 		// the sub-allocated memory on success, or nullptr on failure.
-		inline void* TryAllocate(size_t Size, size_t Alignment)
+		inline FAllocatedSpan TryAllocate(size_t Size, size_t Alignment)
 		{
 			size_t RemainingAligned = AlignDown(Remaining, Alignment);
 			if (RemainingAligned < Size)
 			{
-				return nullptr;
+				return FAllocatedSpan{nullptr, 0};
 			}
-			void* Ptr = reinterpret_cast<std::byte*>(this) + sizeof(FBlockHeader) + BlockDataSize - RemainingAligned;
+			std::byte* Ptr = reinterpret_cast<std::byte*>(this) + sizeof(FBlockHeader) + BlockDataSize - RemainingAligned;
 			Remaining = RemainingAligned - Size;
-			return Ptr;
+			return FAllocatedSpan{Ptr, Size};
+		}
+
+		// How many bytes remain in this block, assuming a given alignment?
+		size_t RemainingAlignedBytes(size_t Alignment)
+		{
+			return AlignDown(Remaining, Alignment);
 		}
 	};
 
 	static_assert(sizeof(FBlockHeader) % DataAlignment == 0);
 
-	size_t NextBlockSize = InlineBlockDataSize * GrowthPercentage / 100;
+	static constexpr size_t DefaultFirstBlockSize = InlineBlockDataSize * GrowthPercentage / 100;
+
+	size_t NextBlockSize = DefaultFirstBlockSize;
 	FBlockHeader* Tail = &InlineBlock;
 	FBlockHeader InlineBlock{/* Prev */ nullptr, InlineBlockDataSize};
 	alignas(BlockAlignment) std::byte Data[InlineBlockDataSize];
diff --git a/Engine/Source/Runtime/AutoRTFM/Private/IntrusivePool.h b/Engine/Source/Runtime/AutoRTFM/Private/IntrusivePool.h
--- a/Engine/Source/Runtime/AutoRTFM/Private/IntrusivePool.h
+++ b/Engine/Source/Runtime/AutoRTFM/Private/IntrusivePool.h
@@ -50,7 +50,7 @@ public:
 		else
 		{
 			// Free list is empty. Allocate another item.
-			void* const Memory = Allocator.Allocate(EntrySize, EntryAlignment);
+			void* const Memory = Allocator.Allocate(EntrySize, EntryAlignment).Data;
 			NumAllocated++;
 
 			Item = new (Memory) FItem(std::forward<ArgumentTypes>(Arguments)...);
diff --git a/Engine/Source/Runtime/AutoRTFM/Private/Pool.h b/Engine/Source/Runtime/AutoRTFM/Private/Pool.h
--- a/Engine/Source/Runtime/AutoRTFM/Private/Pool.h
+++ b/Engine/Source/Runtime/AutoRTFM/Private/Pool.h
@@ -45,7 +45,7 @@ public:
 		else
 		{
 			// Free list is empty. Allocate another item.
-			Memory = Allocator.Allocate(EntrySize, EntryAlignment);
+			Memory = Allocator.Allocate(EntrySize, EntryAlignment).Data;
 			NumAllocated++;
 		}
 
