[AutoRTFM] Add AutoRTFM::IsRetrying method.

This lets an OnAbort or OnComplete handler discover if they are being invoked in response to a cascading retry.

In some cases, such as logging, we want to treat retries differently from transactional failure; failure should print "(Rolled Back) Message" but a retry should print nothing at all. Otherwise the user will see a confusing spray of rolled-back messages that don't correspond to a Verse failure.

#rnx

diff --git a/Engine/Source/Programs/AutoRTFMTests/Private/API.cpp b/Engine/Source/Programs/AutoRTFMTests/Private/API.cpp
--- a/Engine/Source/Programs/AutoRTFMTests/Private/API.cpp
+++ b/Engine/Source/Programs/AutoRTFMTests/Private/API.cpp
@@ -11,21 +11,21 @@ TEST_CASE("API.autortfm_is_transactional")
 {
     REQUIRE(false == autortfm_is_transactional());
 
-    bool InTransaction = false;
-    bool InOpenNest = false;
+    bool bInTransaction = false;
+    bool bInOpenNest = false;
 
     AutoRTFM::Commit([&]
     {
-        InTransaction = autortfm_is_transactional();
+        bInTransaction = autortfm_is_transactional();
 
         AutoRTFM::Open([&]
         {
-            InOpenNest = autortfm_is_transactional();
+            bInOpenNest = autortfm_is_transactional();
         });
     });
 
-    REQUIRE(true == InTransaction);
-    REQUIRE(true == InOpenNest);
+    REQUIRE(true == bInTransaction);
+    REQUIRE(true == bInOpenNest);
 }
 
 TEST_CASE("API.autortfm_is_closed")
@@ -33,56 +33,56 @@ TEST_CASE("API.autortfm_is_closed")
     REQUIRE(false == autortfm_is_closed());
 
     // Set to the opposite of what we expect at the end of function.
-    bool InTransaction = false;
-    bool InOpenNest = true;
-    bool InClosedNestInOpenNest = false;
+    bool bInTransaction = false;
+    bool bInOpenNest = true;
+    bool bInClosedNestInOpenNest = false;
 
     AutoRTFM::Commit([&]
     {
-        InTransaction = autortfm_is_closed();
+        bInTransaction = autortfm_is_closed();
 
         AutoRTFM::Open([&]
         {
-            InOpenNest = autortfm_is_closed();
+            bInOpenNest = autortfm_is_closed();
 
             REQUIRE(AutoRTFM::EContextStatus::OnTrack == AutoRTFM::Close([&]
             {
-                InClosedNestInOpenNest = autortfm_is_closed();
+                bInClosedNestInOpenNest = autortfm_is_closed();
             }));
         });
     });
 
-    REQUIRE(true == InTransaction);
-    REQUIRE(false == InOpenNest);
-    REQUIRE(true == InClosedNestInOpenNest);
+    REQUIRE(true == bInTransaction);
+    REQUIRE(false == bInOpenNest);
+    REQUIRE(true == bInClosedNestInOpenNest);
 }
 
 TEST_CASE("API.autortfm_abort_transaction")
 {
-    bool BeforeNest = false;
-    bool InNest = false;
-    bool AfterNest = false;
+    bool bBeforeNest = false;
+    bool bInNest = false;
+    bool bAfterNest = false;
     
     AutoRTFM::Commit([&]
     {
-        BeforeNest = true;
+        bBeforeNest = true;
 
         AutoRTFM::ETransactionResult Result = AutoRTFM::Transact([&]
         {
             // Because we are aborting this won't actually occur!
-            InNest = true;
+            bInNest = true;
 
             autortfm_abort_transaction();
         });
 
         REQUIRE(AutoRTFM::ETransactionResult::AbortedByRequest == Result);
 
-        AfterNest = true;
+        bAfterNest = true;
     });
 
-    REQUIRE(true == BeforeNest);
-    REQUIRE(false == InNest);
-    REQUIRE(true == AfterNest);
+    REQUIRE(true == bBeforeNest);
+    REQUIRE(false == bInNest);
+    REQUIRE(true == bAfterNest);
 }
 
 TEST_CASE("API.autortfm_open")
@@ -139,20 +139,20 @@ TEST_CASE("API.autortfm_register_open_to_closed_functions")
 
 TEST_CASE("API.autortfm_on_commit")
 {
-    bool OuterTransaction = false;
-    bool InnerTransaction = false;
-    bool InnerTransactionWithAbort = false;
-    bool InnerOpenNest = false;
+    bool bOuterTransaction = false;
+    bool bInnerTransaction = false;
+    bool bInnerTransactionWithAbort = false;
+    bool bInnerOpenNest = false;
 
     AutoRTFM::Commit([&]
     {
         autortfm_on_commit([](void* const Arg)
         {
             *static_cast<bool* const>(Arg) = true;
-        }, &OuterTransaction);
+        }, &bOuterTransaction);
 
         // This should only be modified on the commit!
-        if (OuterTransaction)
+        if (bOuterTransaction)
         {
             AutoRTFM::AbortTransaction();
         }
@@ -162,11 +162,11 @@ TEST_CASE("API.autortfm_on_commit")
             autortfm_on_commit([](void* const Arg)
             {
                 *static_cast<bool* const>(Arg) = true;
-            }, &InnerTransaction);
+            }, &bInnerTransaction);
         });
 
         // This should only be modified on the commit!
-        if (InnerTransaction)
+        if (bInnerTransaction)
         {
             AutoRTFM::AbortTransaction();
         }
@@ -176,7 +176,7 @@ TEST_CASE("API.autortfm_on_commit")
             autortfm_on_commit([](void* const Arg)
             {
                 *static_cast<bool* const>(Arg) = true;
-            }, &InnerTransactionWithAbort);
+            }, &bInnerTransactionWithAbort);
 
             AutoRTFM::AbortTransaction();
         });
@@ -184,7 +184,7 @@ TEST_CASE("API.autortfm_on_commit")
         REQUIRE(AutoRTFM::ETransactionResult::AbortedByRequest == Result);
 
         // This should never be modified because its transaction aborted!
-        if (InnerTransactionWithAbort)
+        if (bInnerTransactionWithAbort)
         {
             AutoRTFM::AbortTransaction();
         }
@@ -194,20 +194,20 @@ TEST_CASE("API.autortfm_on_commit")
             autortfm_on_commit([](void* const Arg)
             {
                 *static_cast<bool* const>(Arg) = true;
-            }, &InnerOpenNest);
+            }, &bInnerOpenNest);
 
             // This should be modified immediately!
-            if (!InnerOpenNest)
+            if (!bInnerOpenNest)
             {
                 AutoRTFM::AbortTransaction();
             }
         });
     });
 
-    REQUIRE(true == OuterTransaction);
-    REQUIRE(true == InnerTransaction);
-    REQUIRE(false == InnerTransactionWithAbort);
-    REQUIRE(true == InnerOpenNest);
+    REQUIRE(true == bOuterTransaction);
+    REQUIRE(true == bInnerTransaction);
+    REQUIRE(false == bInnerTransactionWithAbort);
+    REQUIRE(true == bInnerOpenNest);
 }
 
 TEST_CASE("API.autortfm_on_abort")
@@ -218,29 +218,29 @@ TEST_CASE("API.autortfm_on_abort")
 		return;
 	}
 
-    bool OuterTransaction = false;
-    bool InnerTransaction = false;
-    bool InnerTransactionWithAbort = false;
-    bool InnerOpenNest = false;
+    bool bOuterTransaction = false;
+    bool bInnerTransaction = false;
+    bool bInnerTransactionWithAbort = false;
+    bool bInnerOpenNest = false;
 
     REQUIRE(AutoRTFM::ETransactionResult::Committed == AutoRTFM::Transact([&]
     {
 		// If we are retrying transactions, need to reset the test state.
 		AutoRTFM::OnAbort([&]
 		{
-			OuterTransaction = false;
-			InnerTransaction = false;
-			InnerTransactionWithAbort = false;
-			InnerOpenNest = false;
+			bOuterTransaction = false;
+			bInnerTransaction = false;
+			bInnerTransactionWithAbort = false;
+			bInnerOpenNest = false;
 		});
 
         autortfm_on_abort([](void* const Arg)
         {
             *static_cast<bool* const>(Arg) = true;
-        }, &OuterTransaction);
+        }, &bOuterTransaction);
 
         // This should only be modified on the commit!
-        if (OuterTransaction)
+        if (bOuterTransaction)
         {
             AutoRTFM::AbortTransaction();
         }
@@ -250,11 +250,11 @@ TEST_CASE("API.autortfm_on_abort")
             autortfm_on_abort([](void* const Arg)
             {
                 *static_cast<bool* const>(Arg) = true;
-            }, &InnerTransaction);
+            }, &bInnerTransaction);
         });
 
         // This should only be modified on the commit!
-        if (InnerTransaction)
+        if (bInnerTransaction)
         {
             AutoRTFM::AbortTransaction();
         }
@@ -264,7 +264,7 @@ TEST_CASE("API.autortfm_on_abort")
             autortfm_on_abort([](void* const Arg)
             {
                 *static_cast<bool* const>(Arg) = true;
-            }, &InnerTransactionWithAbort);
+            }, &bInnerTransactionWithAbort);
 
             AutoRTFM::AbortTransaction();
         });
@@ -272,7 +272,7 @@ TEST_CASE("API.autortfm_on_abort")
         REQUIRE(AutoRTFM::ETransactionResult::AbortedByRequest == Result);
 
         // OnAbort runs eagerly on inner abort
-        if (!InnerTransactionWithAbort)
+        if (!bInnerTransactionWithAbort)
         {
             AutoRTFM::AbortTransaction();
         }
@@ -282,20 +282,20 @@ TEST_CASE("API.autortfm_on_abort")
             autortfm_on_abort([](void* const Arg)
             {
                 *static_cast<bool* const>(Arg) = true;
-            }, &InnerOpenNest);
+            }, &bInnerOpenNest);
         });
 
         // This should only be modified on the commit!
-        if (InnerOpenNest)
+        if (bInnerOpenNest)
         {
             AutoRTFM::AbortTransaction();
         }
     }));
 
-    REQUIRE(false == OuterTransaction);
-    REQUIRE(false == InnerTransaction);
-    REQUIRE(true == InnerTransactionWithAbort);
-    REQUIRE(false == InnerOpenNest);
+    REQUIRE(false == bOuterTransaction);
+    REQUIRE(false == bInnerTransaction);
+    REQUIRE(true == bInnerTransactionWithAbort);
+    REQUIRE(false == bInnerOpenNest);
 }
 
 TEST_CASE("API.autortfm_did_allocate")
@@ -373,26 +373,26 @@ TEST_CASE("API.IsTransactional")
 {
     REQUIRE(false == AutoRTFM::IsTransactional());
 
-    bool InTransaction = false;
-    bool InOpenNest = false;
-	bool InAbort = true;
-	bool InCommit = true;
-	bool InComplete = true;
+    bool bInTransaction = false;
+    bool bInOpenNest = false;
+	bool bInAbort = true;
+	bool bInCommit = true;
+	bool bInComplete = true;
 
     AutoRTFM::Testing::Commit([&]
     {
-        InTransaction = AutoRTFM::IsTransactional();
+        bInTransaction = AutoRTFM::IsTransactional();
 
         AutoRTFM::Open([&]
         {
-            InOpenNest = AutoRTFM::IsTransactional();
+            bInOpenNest = AutoRTFM::IsTransactional();
         });
 
 		AutoRTFM::Testing::Abort([&]
 		{
 			AutoRTFM::OnAbort([&]
 			{
-				InAbort = AutoRTFM::IsTransactional();
+				bInAbort = AutoRTFM::IsTransactional();
 			});
 
 			AutoRTFM::AbortTransaction();
@@ -400,20 +400,20 @@ TEST_CASE("API.IsTransactional")
 
 		AutoRTFM::OnCommit([&]
 		{
-			InCommit = AutoRTFM::IsTransactional();
+			bInCommit = AutoRTFM::IsTransactional();
 		});
 
 		AutoRTFM::OnComplete([&]
 		{
-			InComplete = AutoRTFM::IsTransactional();
+			bInComplete = AutoRTFM::IsTransactional();
 		});
     });
 
-    REQUIRE(true == InTransaction);
-    REQUIRE(true == InOpenNest);
-	REQUIRE(false == InAbort);
-	REQUIRE(false == InCommit);
-	REQUIRE(false == InComplete);
+    REQUIRE(true == bInTransaction);
+    REQUIRE(true == bInOpenNest);
+	REQUIRE(false == bInAbort);
+	REQUIRE(false == bInCommit);
+	REQUIRE(false == bInComplete);
 }
 
 TEST_CASE("API.IsClosed")
@@ -421,22 +421,22 @@ TEST_CASE("API.IsClosed")
     REQUIRE(false == AutoRTFM::IsClosed());
 
     // Set to the opposite of what we expect at the end of function.
-    bool InTransaction = false;
-    bool InOpenNest = true;
-    bool InClosedNestInOpenNest = false;
-	bool InAbort = true;
-	bool InCommit = true;
-	bool InComplete = true;
+    bool bInTransaction = false;
+    bool bInOpenNest = true;
+    bool bInClosedNestInOpenNest = false;
+	bool bInAbort = true;
+	bool bInCommit = true;
+	bool bInComplete = true;
 
     AutoRTFM::Testing::Commit([&]
     {
-        InTransaction = AutoRTFM::IsClosed();
+        bInTransaction = AutoRTFM::IsClosed();
 
 		AutoRTFM::Testing::Abort([&]
 		{
 			AutoRTFM::OnAbort([&]
 			{
-				InAbort = AutoRTFM::IsClosed();
+				bInAbort = AutoRTFM::IsClosed();
 			});
 
 			AutoRTFM::AbortTransaction();
@@ -444,31 +444,31 @@ TEST_CASE("API.IsClosed")
 
 		AutoRTFM::OnCommit([&]
 		{
-			InCommit = AutoRTFM::IsClosed();
+			bInCommit = AutoRTFM::IsClosed();
 		});
 
 		AutoRTFM::OnComplete([&]
 		{
-			InComplete = AutoRTFM::IsClosed();
+			bInComplete = AutoRTFM::IsClosed();
 		});
 
         AutoRTFM::Open([&]
         {
-            InOpenNest = AutoRTFM::IsClosed();
+            bInOpenNest = AutoRTFM::IsClosed();
 
             REQUIRE(AutoRTFM::EContextStatus::OnTrack == AutoRTFM::Close([&]
             {
-                InClosedNestInOpenNest = AutoRTFM::IsClosed();
+                bInClosedNestInOpenNest = AutoRTFM::IsClosed();
             }));
         });
     });
 
-    REQUIRE(true == InTransaction);
-    REQUIRE(false == InOpenNest);
-    REQUIRE(true == InClosedNestInOpenNest);
-	REQUIRE(false == InAbort);
-	REQUIRE(false == InCommit);
-	REQUIRE(false == InComplete);
+    REQUIRE(true == bInTransaction);
+    REQUIRE(false == bInOpenNest);
+    REQUIRE(true == bInClosedNestInOpenNest);
+	REQUIRE(false == bInAbort);
+	REQUIRE(false == bInCommit);
+	REQUIRE(false == bInComplete);
 }
 
 TEST_CASE("API.IsCommittingOrAborting")
@@ -476,22 +476,24 @@ TEST_CASE("API.IsCommittingOrAborting")
 	REQUIRE(false == AutoRTFM::IsCommittingOrAborting());
 
 	// Set to the opposite of what we expect at the end of function.
-	bool InTransaction = true;
-	bool InOpenNest = true;
-	bool InClosedNestInOpenNest = true;
-	bool InAbort = false;
-	bool InCommit = false;
-	bool InComplete = false;
+	bool bInTransaction = true;
+	bool bInOpenNest = true;
+	bool bInClosedNestInOpenNest = true;
+	bool bInAbort = false;
+	bool bInCommit = false;
+	bool bInComplete = false;
+	bool bInRetryAbort = false;
+	bool bInRetryComplete = false;
 
 	AutoRTFM::Testing::Commit([&]
 	{
-		InTransaction = AutoRTFM::IsCommittingOrAborting();
+		bInTransaction = AutoRTFM::IsCommittingOrAborting();
 
 		AutoRTFM::Testing::Abort([&]
 		{
 			AutoRTFM::OnAbort([&]
 			{
-				InAbort = AutoRTFM::IsCommittingOrAborting();
+				bInAbort = AutoRTFM::IsCommittingOrAborting();
 			});
 
 			AutoRTFM::AbortTransaction();
@@ -499,31 +501,133 @@ TEST_CASE("API.IsCommittingOrAborting")
 
 		AutoRTFM::OnCommit([&]
 		{
-			InCommit = AutoRTFM::IsCommittingOrAborting();
+			bInCommit = AutoRTFM::IsCommittingOrAborting();
 		});
 
 		AutoRTFM::OnComplete([&]
 		{
-			InComplete = AutoRTFM::IsCommittingOrAborting();
+			bInComplete = AutoRTFM::IsCommittingOrAborting();
 		});
 
 		AutoRTFM::Open([&]
 		{
-			InOpenNest = AutoRTFM::IsCommittingOrAborting();
+			bInOpenNest = AutoRTFM::IsCommittingOrAborting();
 
 			REQUIRE(AutoRTFM::EContextStatus::OnTrack == AutoRTFM::Close([&]
 			{
-				InClosedNestInOpenNest = AutoRTFM::IsCommittingOrAborting();
+				bInClosedNestInOpenNest = AutoRTFM::IsCommittingOrAborting();
 			}));
 		});
 	});
 
-	REQUIRE(false == InTransaction);
-	REQUIRE(false == InOpenNest);
-	REQUIRE(false == InClosedNestInOpenNest);
-	REQUIRE(true == InAbort);
-	REQUIRE(true == InCommit);
-	REQUIRE(true == InComplete);
+	bool bAlreadyRetried = false;
+
+	AutoRTFM::Testing::Commit([&]
+	{
+		if (!bAlreadyRetried)
+		{
+			AutoRTFM::OnAbort([&]
+			{
+				bInRetryAbort = AutoRTFM::IsCommittingOrAborting();
+			});
+			AutoRTFM::OnComplete([&]
+			{
+				bInRetryComplete = AutoRTFM::IsCommittingOrAborting();
+				bAlreadyRetried = true;
+			});
+
+			AutoRTFM::CascadingRetryTransaction();
+		}
+	});
+
+	REQUIRE(false == bInTransaction);
+	REQUIRE(false == bInOpenNest);
+	REQUIRE(false == bInClosedNestInOpenNest);
+	REQUIRE(true == bInAbort);
+	REQUIRE(true == bInCommit);
+	REQUIRE(true == bInComplete);
+	REQUIRE(true == bInRetryAbort);
+	REQUIRE(true == bInRetryComplete);
+}
+
+
+TEST_CASE("API.IsRetrying")
+{
+	REQUIRE(false == AutoRTFM::IsRetrying());
+
+	// Set to the opposite of what we expect at the end of function.
+	bool bInTransaction = true;
+	bool bInOpenNest = true;
+	bool bInClosedNestInOpenNest = true;
+	bool bInAbort = true;
+	bool bInCommit = true;
+	bool bInComplete = true;
+	bool bInRetryAbort = false;
+	bool bInRetryComplete = false;
+
+	AutoRTFM::Testing::Commit([&]
+	{
+		bInTransaction = AutoRTFM::IsRetrying();
+
+		AutoRTFM::Testing::Abort([&]
+		{
+			AutoRTFM::OnAbort([&]
+			{
+				bInAbort = AutoRTFM::IsRetrying();
+			});
+
+			AutoRTFM::AbortTransaction();
+		});
+
+		AutoRTFM::OnCommit([&]
+		{
+			bInCommit = AutoRTFM::IsRetrying();
+		});
+
+		AutoRTFM::OnComplete([&]
+		{
+			bInComplete = AutoRTFM::IsRetrying();
+		});
+
+		AutoRTFM::Open([&]
+		{
+			bInOpenNest = AutoRTFM::IsRetrying();
+
+			REQUIRE(AutoRTFM::EContextStatus::OnTrack == AutoRTFM::Close([&]
+			{
+				bInClosedNestInOpenNest = AutoRTFM::IsRetrying();
+			}));
+		});
+	});
+
+	bool bAlreadyRetried = false;
+
+	AutoRTFM::Testing::Commit([&]
+	{
+		if (!bAlreadyRetried)
+		{
+			AutoRTFM::OnAbort([&]
+			{
+				bInRetryAbort = AutoRTFM::IsRetrying();
+			});
+			AutoRTFM::OnComplete([&]
+			{
+				bInRetryComplete = AutoRTFM::IsRetrying();
+				bAlreadyRetried = true;
+			});
+
+			AutoRTFM::CascadingRetryTransaction();
+		}
+	});
+
+	REQUIRE(false == bInTransaction);
+	REQUIRE(false == bInOpenNest);
+	REQUIRE(false == bInClosedNestInOpenNest);
+	REQUIRE(false == bInAbort);
+	REQUIRE(false == bInCommit);
+	REQUIRE(false == bInComplete);
+	REQUIRE(true == bInRetryAbort);
+	REQUIRE(true == bInRetryComplete);
 }
 
 TEST_CASE("API.Transact")
@@ -579,30 +683,30 @@ TEST_CASE("API.Commit")
 
 TEST_CASE("API.Abort")
 {
-    bool BeforeNest = false;
-    bool InNest = false;
-    bool AfterNest = false;
+    bool bBeforeNest = false;
+    bool bInNest = false;
+    bool bAfterNest = false;
 
     AutoRTFM::Commit([&]
     {
-        BeforeNest = true;
+        bBeforeNest = true;
 
         AutoRTFM::ETransactionResult Result = AutoRTFM::Transact([&]
         {
             // Because we are aborting this won't actually occur!
-            InNest = true;
+            bInNest = true;
 
             AutoRTFM::AbortTransaction();
         });
 
         REQUIRE(AutoRTFM::ETransactionResult::AbortedByRequest == Result);
 
-        AfterNest = true;
+        bAfterNest = true;
     });
 
-    REQUIRE(true == BeforeNest);
-    REQUIRE(false == InNest);
-    REQUIRE(true == AfterNest);
+    REQUIRE(true == bBeforeNest);
+    REQUIRE(false == bInNest);
+    REQUIRE(true == bAfterNest);
 }
 
 TEST_CASE("API.Open")
@@ -664,25 +768,25 @@ TEST_CASE("API.OpenMacro_WithAbort")
 
 TEST_CASE("API.Close")
 {
-    bool InClosedNest = false;
-    bool InOpenNest = false;
-    bool InClosedNestInOpenNest = false;
+    bool bInClosedNest = false;
+    bool bInOpenNest = false;
+    bool bInClosedNestInOpenNest = false;
 
     AutoRTFM::ETransactionResult Result = AutoRTFM::Transact([&]
     {
         // A closed call inside a transaction does not abort.
-        AutoRTFM::EContextStatus CloseStatusA = AutoRTFM::Close([&] { InClosedNest = true; });
+        AutoRTFM::EContextStatus CloseStatusA = AutoRTFM::Close([&] { bInClosedNest = true; });
         REQUIRE(AutoRTFM::EContextStatus::OnTrack == CloseStatusA);
 
         AutoRTFM::Open([&]
         {
             AutoRTFM::EContextStatus CloseStatusB = AutoRTFM::Close([&]
             {
-                InClosedNestInOpenNest = true;
+                bInClosedNestInOpenNest = true;
             });
             REQUIRE(AutoRTFM::EContextStatus::OnTrack == CloseStatusB);
 
-            InOpenNest = true;
+            bInOpenNest = true;
         });
 
         AutoRTFM::AbortTransaction();
@@ -690,27 +794,27 @@ TEST_CASE("API.Close")
 
     REQUIRE(AutoRTFM::ETransactionResult::AbortedByRequest == Result);
 
-    REQUIRE(false == InClosedNest);
-    REQUIRE(true == InOpenNest);
-    REQUIRE(false == InClosedNestInOpenNest);
+    REQUIRE(false == bInClosedNest);
+    REQUIRE(true == bInOpenNest);
+    REQUIRE(false == bInClosedNestInOpenNest);
 }
 
 TEST_CASE("API.OnCommit")
 {
-    bool OuterTransaction = false;
-    bool InnerTransaction = false;
-    bool InnerTransactionWithAbort = false;
-    bool InnerOpenNest = false;
+    bool bOuterTransaction = false;
+    bool bInnerTransaction = false;
+    bool bInnerTransactionWithAbort = false;
+    bool bInnerOpenNest = false;
 
     REQUIRE(AutoRTFM::ETransactionResult::Committed == AutoRTFM::Transact([&]
     {
         AutoRTFM::OnCommit([&]
         {
-            OuterTransaction = true;
+            bOuterTransaction = true;
         });
 
         // This should only be modified on the commit!
-        if (OuterTransaction)
+        if (bOuterTransaction)
         {
             AutoRTFM::AbortTransaction();
         }
@@ -719,12 +823,12 @@ TEST_CASE("API.OnCommit")
         {
             AutoRTFM::OnCommit([&]
             {
-                InnerTransaction = true;
+                bInnerTransaction = true;
             });
         });
 
         // This should only be modified on the commit!
-        if (InnerTransaction)
+        if (bInnerTransaction)
         {
 			AutoRTFM::AbortTransaction();
         }
@@ -734,7 +838,7 @@ TEST_CASE("API.OnCommit")
         {
             AutoRTFM::OnCommit([&]
             {
-                InnerTransactionWithAbort = true;
+                bInnerTransactionWithAbort = true;
             });
 
             AutoRTFM::AbortTransaction();
@@ -743,7 +847,7 @@ TEST_CASE("API.OnCommit")
         REQUIRE(AutoRTFM::ETransactionResult::AbortedByRequest == Result);
 
         // This should never be modified because its transaction aborted!
-        if (InnerTransactionWithAbort)
+        if (bInnerTransactionWithAbort)
         {
 			AutoRTFM::AbortTransaction();
         }
@@ -752,21 +856,21 @@ TEST_CASE("API.OnCommit")
         {
             AutoRTFM::OnCommit([&]
             {
-                InnerOpenNest = true;
+                bInnerOpenNest = true;
             });
 
             // This should be modified immediately!
-            if (!InnerOpenNest)
+            if (!bInnerOpenNest)
             {
 				AutoRTFM::AbortTransaction();
             }
         });
     }));
 
-    REQUIRE(true == OuterTransaction);
-    REQUIRE(true == InnerTransaction);
-    REQUIRE(false == InnerTransactionWithAbort);
-    REQUIRE(true == InnerOpenNest);
+    REQUIRE(true == bOuterTransaction);
+    REQUIRE(true == bInnerTransaction);
+    REQUIRE(false == bInnerTransactionWithAbort);
+    REQUIRE(true == bInnerOpenNest);
 }
 
 TEST_CASE("API.OnComplete")
@@ -953,29 +1057,29 @@ TEST_CASE("API.OnAbort")
 		return;
 	}
 
-    bool OuterTransaction = false;
-    bool InnerTransaction = false;
-    bool InnerTransactionWithAbort = false;
-    bool InnerOpenNest = false;
+    bool bOuterTransaction = false;
+    bool bInnerTransaction = false;
+    bool bInnerTransactionWithAbort = false;
+    bool bInnerOpenNest = false;
 
     REQUIRE(AutoRTFM::ETransactionResult::Committed == AutoRTFM::Transact([&]
     {
 		// If we are retrying transactions, need to reset the test state.
 		AutoRTFM::OnAbort([&]
 		{
-			OuterTransaction = false;
-			InnerTransaction = false;
-			InnerTransactionWithAbort = false;
-			InnerOpenNest = false;
+			bOuterTransaction = false;
+			bInnerTransaction = false;
+			bInnerTransactionWithAbort = false;
+			bInnerOpenNest = false;
 		});
 
         AutoRTFM::OnAbort([&]
         {
-            OuterTransaction = true;
+            bOuterTransaction = true;
         });
 
         // This should only be modified on the commit!
-        if (OuterTransaction)
+        if (bOuterTransaction)
         {
 			AutoRTFM::AbortTransaction();
         }
@@ -984,12 +1088,12 @@ TEST_CASE("API.OnAbort")
         {
             AutoRTFM::OnAbort([&]
             {
-                InnerTransaction = true;
+                bInnerTransaction = true;
             });
         });
 
         // This should only be modified on the commit!
-        if (InnerTransaction)
+        if (bInnerTransaction)
         {
 			AutoRTFM::AbortTransaction();
         }
@@ -998,7 +1102,7 @@ TEST_CASE("API.OnAbort")
         {
             AutoRTFM::OnAbort([&]
             {
-                InnerTransactionWithAbort = true;
+                bInnerTransactionWithAbort = true;
             });
 
 			AutoRTFM::AbortTransaction();
@@ -1007,7 +1111,7 @@ TEST_CASE("API.OnAbort")
         REQUIRE(AutoRTFM::ETransactionResult::AbortedByRequest == Result);
 
         // Inner OnAbort runs eagerly
-        if (!InnerTransactionWithAbort)
+        if (!bInnerTransactionWithAbort)
         {
 			AutoRTFM::AbortTransaction();
         }
@@ -1016,21 +1120,21 @@ TEST_CASE("API.OnAbort")
         {
             AutoRTFM::OnAbort([&]
             {
-                InnerOpenNest = true;
+                bInnerOpenNest = true;
             });
         });
 
         // This should only be modified on the commit!
-        if (InnerOpenNest)
+        if (bInnerOpenNest)
         {
 			AutoRTFM::AbortTransaction();
         }
     }));
 
-    REQUIRE(false == OuterTransaction);
-    REQUIRE(false == InnerTransaction);
-    REQUIRE(true == InnerTransactionWithAbort);
-    REQUIRE(false == InnerOpenNest);
+    REQUIRE(false == bOuterTransaction);
+    REQUIRE(false == bInnerTransaction);
+    REQUIRE(true == bInnerTransactionWithAbort);
+    REQUIRE(false == bInnerOpenNest);
 }
 
 TEST_CASE("API.OnAbort_MutableCapture")
diff --git a/Engine/Source/Runtime/AutoRTFM/Private/API.cpp b/Engine/Source/Runtime/AutoRTFM/Private/API.cpp
--- a/Engine/Source/Runtime/AutoRTFM/Private/API.cpp
+++ b/Engine/Source/Runtime/AutoRTFM/Private/API.cpp
@@ -763,8 +763,10 @@ extern "C" bool autortfm_is_transactional() noexcept
 
 	if (ForTheRuntime::IsAutoRTFMRuntimeEnabled())
 	{
-		FContext* Context = FContext::Get();
-		return Context && Context->IsTransactional();
+		if (FContext* Context = FContext::Get())
+		{
+			return Context->IsTransactional();
+		}
 	}
 
 	return false;
@@ -774,8 +776,23 @@ extern "C" UE_AUTORTFM_ALWAYS_OPEN bool autortfm_is_committing_or_aborting() noe
 {
 	if (ForTheRuntime::IsAutoRTFMRuntimeEnabled())
 	{
-		FContext* Context = FContext::Get();
-		return Context && Context->IsCommittingOrAborting();
+		if (FContext* Context = FContext::Get())
+		{
+			return Context->IsCommittingOrAborting();
+		}
+	}
+
+	return false;
+}
+
+extern "C" UE_AUTORTFM_ALWAYS_OPEN bool autortfm_is_retrying() noexcept
+{
+	if (ForTheRuntime::IsAutoRTFMRuntimeEnabled())
+	{
+		if (FContext* Context = FContext::Get())
+		{
+			return Context->IsRetrying();
+		}
 	}
 
 	return false;
diff --git a/Engine/Source/Runtime/AutoRTFM/Private/Context.cpp b/Engine/Source/Runtime/AutoRTFM/Private/Context.cpp
--- a/Engine/Source/Runtime/AutoRTFM/Private/Context.cpp
+++ b/Engine/Source/Runtime/AutoRTFM/Private/Context.cpp
@@ -66,6 +66,11 @@ bool FContext::IsCommittingOrAborting() const
 	}
 }
 
+bool FContext::IsRetrying() const
+{
+	return GetStatus() == EContextStatus::AbortedByCascadingRetry;
+}
+
 void FContext::MaterializeDeferredTransactions()
 {
 	uint64_t NumToAllocate = GetNumDeferredTransactions();
diff --git a/Engine/Source/Runtime/AutoRTFM/Private/Context.h b/Engine/Source/Runtime/AutoRTFM/Private/Context.h
--- a/Engine/Source/Runtime/AutoRTFM/Private/Context.h
+++ b/Engine/Source/Runtime/AutoRTFM/Private/Context.h
@@ -47,6 +47,7 @@ public:
 	bool IsAborting() const;
 	bool IsTransactional() const;
 	bool IsCommittingOrAborting() const;
+	bool IsRetrying() const;
 
     // Record that a write is about to occur at the given LogicalAddress of Size bytes.
     void RecordWrite(void* LogicalAddress, size_t Size);
diff --git a/Engine/Source/Runtime/AutoRTFM/Public/AutoRTFM.h b/Engine/Source/Runtime/AutoRTFM/Public/AutoRTFM.h
--- a/Engine/Source/Runtime/AutoRTFM/Public/AutoRTFM.h
+++ b/Engine/Source/Runtime/AutoRTFM/Public/AutoRTFM.h
@@ -183,6 +183,15 @@ UE_AUTORTFM_CRITICAL_INLINE bool autortfm_is_committing_or_aborting(void) AUTORT
 }
 #endif
 
+#if UE_AUTORTFM_ENABLED
+UE_AUTORTFM_API bool autortfm_is_retrying(void) AUTORTFM_NOEXCEPT;
+#else
+UE_AUTORTFM_CRITICAL_INLINE bool autortfm_is_retrying(void) AUTORTFM_NOEXCEPT
+{
+	return false;
+}
+#endif
+
 #if UE_AUTORTFM_ENABLED
 UE_AUTORTFM_API autortfm_result autortfm_transact(void (*UninstrumentedWork)(void*), void (*InstrumentedWork)(void*), void* Arg) AUTORTFM_EXCEPT;
 #else
@@ -579,9 +588,13 @@ UE_AUTORTFM_CRITICAL_INLINE_ALWAYS bool IsTransactional() { return autortfm_is_t
 UE_AUTORTFM_CRITICAL_INLINE_ALWAYS bool IsClosed() { return autortfm_is_closed(); }
 
 // Tells us if we are currently committing or aborting a transaction. This will return true
-// in an on-abort or on-commit.
+// when inside an on-abort, on-commit or on-complete handler.
 UE_AUTORTFM_CRITICAL_INLINE bool IsCommittingOrAborting() { return autortfm_is_committing_or_aborting(); }
 
+// Tells us if we are currently retrying a transaction. This will return true when inside an
+// on-abort or on-complete handler that was triggered by CascadingRetryTransaction. 
+UE_AUTORTFM_CRITICAL_INLINE bool IsRetrying() { return autortfm_is_retrying(); }
+
 // Returns true if the passed-in pointer is on the stack of the currently-executing transaction.
 // This is occasionally necessary when writing OnAbort handlers for objects on the stack, since
 // we don't want to scribble on stack memory that might have been reused.
