Apply the fix that Richard Malo discovered for the transactionally safe loader and write some tests to prove it works in our transactions.

I've guarded this new behaviour behind the `TransactionallySafeAsyncLoading.LegacyTrashedPackages` CVar - just out of precaution.

diff --git a/Engine/Source/Programs/AutoRTFMTests/Private/AutoRTFMTests.cpp b/Engine/Source/Programs/AutoRTFMTests/Private/AutoRTFMTests.cpp
--- a/Engine/Source/Programs/AutoRTFMTests/Private/AutoRTFMTests.cpp
+++ b/Engine/Source/Programs/AutoRTFMTests/Private/AutoRTFMTests.cpp
@@ -5,8 +5,10 @@
 #include "CoreMinimal.h"
 #include "Misc/EventPool.h"
 #include "Misc/LazySingleton.h"
+#include "Misc/PackageName.h"
 #include "RequiredProgramMainCPPInclude.h" // required for ue programs
 #include "HAL/MallocLeakDetection.h"
+#include "Serialization/PackageStore.h"
 
 IMPLEMENT_APPLICATION(AutoRTFMTests, "AutoRTFMTests");
 
@@ -116,6 +118,33 @@ int RunTests(int ArgC, const char* ArgV[])
 	// Don't print memory validation stats.
 	AutoRTFM::ForTheRuntime::SetMemoryValidationStatisticsEnabled(false);
 
+	// A fake package store backend so we can actually load async packages for testing.
+	struct MyPackageStoreBackend final : IPackageStoreBackend
+	{
+		void OnMounted(TSharedRef<const FPackageStoreBackendContext> InContext) override {} 
+		void BeginRead() override {}
+		void EndRead() override {} 
+
+		EPackageStoreEntryStatus GetPackageStoreEntry(FPackageId PackageId, FName PackageName,
+			FPackageStoreEntry& OutPackageStoreEntry) override
+		{
+			if (!FPackageName::DoesPackageExist(PackageName.ToString()))
+			{
+				return EPackageStoreEntryStatus::Missing;
+			}
+
+			OutPackageStoreEntry.LoaderType = EPackageLoader::LinkerLoad;
+			OutPackageStoreEntry.LinkerLoadCaseCorrectedPackageName = PackageName;
+			
+			return EPackageStoreEntryStatus::Ok;
+		}
+
+		bool GetPackageRedirectInfo(FPackageId PackageId, FName& OutSourcePackageName, FPackageId& OutRedirectedToPackageId) override { return false; }
+	};
+
+	TSharedPtr<MyPackageStoreBackend> PackageStoreBackend = MakeShared<MyPackageStoreBackend>();
+	FPackageStore::Get().Mount(PackageStoreBackend.ToSharedRef());
+
 	FModuleManager::Get().StartProcessingNewlyLoadedObjects();
 
 	int Result = Session.run();
diff --git a/Engine/Source/Programs/AutoRTFMTests/Private/UECoreTests.cpp b/Engine/Source/Programs/AutoRTFMTests/Private/UECoreTests.cpp
--- a/Engine/Source/Programs/AutoRTFMTests/Private/UECoreTests.cpp
+++ b/Engine/Source/Programs/AutoRTFMTests/Private/UECoreTests.cpp
@@ -23,7 +23,6 @@
 #include "Serialization/CustomVersion.h"
 #include "UObject/CoreRedirects.h"
 #include "UObject/DynamicallyTypedValue.h"
-#include "UObject/LinkerInstancingContext.h"
 #include "UObject/NameTypes.h"
 #include "UObject/Stack.h"
 #include "UObject/UObjectArray.h"
@@ -1725,252 +1724,6 @@ TEST_CASE("UECore.FOutputDeviceRedirector")
 	}
 }
 
-TEST_CASE("UECore.AsyncLoading")
-{
-	SECTION("DoesPackageExist")
-	{
-		bool Exists = false;
-
-		AutoRTFM::Testing::Commit([&]
-		{
-			FString Name(FString::Printf(TEXT("/Game/AutoRTFMTestPackage%d"), __LINE__));
-			Exists = FPackageName::DoesPackageExist(Name);
-		});
-
-		REQUIRE(!Exists);
-	}
-
-	SECTION("LoadPackageAsync")
-	{
-		int32 RequestId = -1;
-
-		AutoRTFM::Testing::Commit([&]
-		{
-			FString Name(FString::Printf(TEXT("/AutoRTFMTestPackage%d"), __LINE__));
-			RequestId = LoadPackageAsync(Name);
-		});
-
-		FlushAsyncLoading(RequestId);
-	}
-
-	SECTION("IsAsyncLoading")
-	{
-		int32 RequestId = -1;
-
-		AutoRTFM::Testing::Commit([&]
-		{
-			REQUIRE(!IsAsyncLoading());
-
-			FString Name(FString::Printf(TEXT("/AutoRTFMTestPackage%d"), __LINE__));
-			RequestId = LoadPackageAsync(Name);
-
-			REQUIRE(IsAsyncLoading());
-		});
-
-		FlushAsyncLoading(RequestId);
-	}
-
-	SECTION("FlushAsyncLoading")
-	{
-		FString Name(FString::Printf(TEXT("/AutoRTFMTestPackage%d"), __LINE__));
-		int32 RequestId = LoadPackageAsync(Name);
-
-		AutoRTFM::Testing::Commit([&]
-		{
-			FlushAsyncLoading(RequestId);
-		});
-	}
-
-	SECTION("FlushAsyncLoading Empty")
-	{
-		FString Name(FString::Printf(TEXT("/AutoRTFMTestPackage%d"), __LINE__));
-		int32 RequestId = LoadPackageAsync(Name);
-
-		AutoRTFM::Testing::Commit([&]
-		{
-			FlushAsyncLoading();
-		});
-	}
-
-	SECTION("FlushAsyncLoading One In One Out")
-	{
-		FString Name(FString::Printf(TEXT("/AutoRTFMTestPackage%d"), __LINE__));
-		int32 RequestId1 = LoadPackageAsync(Name);
-
-		AutoRTFM::Testing::Commit([&]
-		{
-			FString Name(FString::Printf(TEXT("/AutoRTFMTestPackage%d"), __LINE__));
-			int32 RequestId2 = LoadPackageAsync(Name);
-			TArray<int32> RequestIds;
-			RequestIds.Add(RequestId1);
-			RequestIds.Add(RequestId2);
-			FlushAsyncLoading(RequestIds);
-		});
-	}
-
-	SECTION("CompletionDelegate is called closed")
-	{
-		AutoRTFM::Testing::Abort([&]
-		{
-			FString Name(FString::Printf(TEXT("/AutoRTFMTestPackage%d"), __LINE__));
-			FLoadPackageAsyncDelegate CompletionDelegate;
-			CompletionDelegate.BindLambda([&](const FName&, UPackage*, EAsyncLoadingResult::Type) { REQUIRE(AutoRTFM::IsClosed()); });
-			int RequestId = LoadPackageAsync(Name, CompletionDelegate);
-			FlushAsyncLoading(RequestId);
-			AutoRTFM::AbortTransaction();
-		});
-	}
-	
-	SECTION("CompletionDelegate aborts")
-	{
-		AutoRTFM::Testing::Abort([&]
-		{
-			FString Name(FString::Printf(TEXT("/AutoRTFMTestPackage%d"), __LINE__));
-			FLoadPackageAsyncDelegate CompletionDelegate;
-			CompletionDelegate.BindLambda([&](const FName&, UPackage*, EAsyncLoadingResult::Type) { AutoRTFM::AbortTransaction(); });
-			int RequestId = LoadPackageAsync(Name, CompletionDelegate);
-			FlushAsyncLoading(RequestId);
-			FAIL("Unreachable!");
-		});
-	}
-	
-	SECTION("FLoadPackageAsyncOptionalParams::CompletionDelegate is called closed")
-	{
-		AutoRTFM::Testing::Abort([&]
-		{
-			FString Name(FString::Printf(TEXT("/AutoRTFMTestPackage%d"), __LINE__));
-			FLoadPackageAsyncOptionalParams Params;
-			Params.CompletionDelegate.Reset(new FLoadPackageAsyncDelegate());
-			Params.CompletionDelegate->BindLambda([&](const FName&, UPackage*, EAsyncLoadingResult::Type) { REQUIRE(AutoRTFM::IsClosed()); });
-			int RequestId = LoadPackageAsync(Name, MoveTemp(Params));
-			FlushAsyncLoading(RequestId);
-			AutoRTFM::AbortTransaction();
-		});
-	}
-	
-	SECTION("FLoadPackageAsyncOptionalParams::CompletionDelegate aborts")
-	{
-		AutoRTFM::Testing::Abort([&]
-		{
-			FString Name(FString::Printf(TEXT("/AutoRTFMTestPackage%d"), __LINE__));
-			FLoadPackageAsyncOptionalParams Params;
-			Params.CompletionDelegate.Reset(new FLoadPackageAsyncDelegate());
-			Params.CompletionDelegate->BindLambda([&](const FName&, UPackage*, EAsyncLoadingResult::Type) { AutoRTFM::AbortTransaction(); });
-			int RequestId = LoadPackageAsync(Name, MoveTemp(Params));
-			FlushAsyncLoading(RequestId);
-			FAIL("Unreachable!");
-		});
-	}
-	
-	SECTION("FLoadPackageAsyncOptionalParams::CompletionDelegate creates UObject")
-	{
-		UMyAutoRTFMTestObject* OpenObject = nullptr;
-		UMyAutoRTFMTestObject* ClosedObject = nullptr;
-
-		AutoRTFM::Testing::Abort([&]
-		{
-			ClosedObject = NewObject<UMyAutoRTFMTestObject>();
-			FString Name(FString::Printf(TEXT("/AutoRTFMTestPackage%d"), __LINE__));
-			FLoadPackageAsyncOptionalParams Params;
-			Params.CompletionDelegate.Reset(new FLoadPackageAsyncDelegate());
-			Params.CompletionDelegate->BindLambda([&](const FName&, UPackage*, EAsyncLoadingResult::Type)
-			{
-				OpenObject = NewObject<UMyAutoRTFMTestObject>();
-				AutoRTFM::AbortTransaction();
-			});
-			int RequestId = LoadPackageAsync(Name, MoveTemp(Params));
-			FlushAsyncLoading(RequestId);
-			FAIL("Unreachable!");
-		});
-
-		REQUIRE(nullptr == ClosedObject);
-		REQUIRE(nullptr == OpenObject);
-	}
-	
-	SECTION("FLoadPackageAsyncOptionalParams::CompletionDelegate calls another LoadPackageAsync")
-	{
-		AutoRTFM::Testing::Abort([&]
-		{
-			FLoadPackageAsyncOptionalParams Params;
-			Params.CompletionDelegate.Reset(new FLoadPackageAsyncDelegate());
-			Params.CompletionDelegate->BindLambda([&](const FName&, UPackage*, EAsyncLoadingResult::Type)
-			{
-				int RequestId = LoadPackageAsync(FString::Printf(TEXT("/AutoRTFMTestPackage%d"), __LINE__), MoveTemp(Params));
-				AutoRTFM::AbortTransaction();
-			});
-			int RequestId = LoadPackageAsync(FString::Printf(TEXT("/AutoRTFMTestPackage%d"), __LINE__), MoveTemp(Params));
-			FlushAsyncLoading(RequestId);
-			FAIL("Unreachable!");
-		});
-	}
-
-	SECTION("Multiple retries because of multiple loads with commit")
-	{
-		AutoRTFMTestUtils::FScopedRetry Retry(AutoRTFM::ForTheRuntime::EAutoRTFMRetryTransactionState::NoRetry);
-
-		int NumCompletionCallbacks = 0;
-		AutoRTFM::Testing::Commit([&]
-		{
-			FLoadPackageAsyncDelegate CompletionDelegate;
-			CompletionDelegate.BindLambda([&](const FName&, UPackage*, EAsyncLoadingResult::Type) 
-			{
-				// Do this open so we can check how many retries occurred.
-				AutoRTFM::Open([&] { NumCompletionCallbacks++; });
-			});
-
-			TArray<int32> RequestIds;
-			RequestIds.Add(LoadPackageAsync(FString::Printf(TEXT("/AutoRTFMTestPackage%d"), __LINE__), CompletionDelegate));
-			RequestIds.Add(LoadPackageAsync(FString::Printf(TEXT("/AutoRTFMTestPackage%d"), __LINE__), CompletionDelegate));
-			RequestIds.Add(LoadPackageAsync(FString::Printf(TEXT("/AutoRTFMTestPackage%d"), __LINE__), CompletionDelegate));
-
-			FlushAsyncLoading(RequestIds);
-			REQUIRE(3 == NumCompletionCallbacks);
-		});
-	}
-
-	SECTION("Multiple retries because of multiple loads with abort")
-	{
-		AutoRTFMTestUtils::FScopedRetry Retry(AutoRTFM::ForTheRuntime::EAutoRTFMRetryTransactionState::NoRetry);
-
-		int NumCompletionCallbacks = 0;
-		AutoRTFM::Testing::Abort([&]
-		{
-			FLoadPackageAsyncDelegate CompletionDelegate;
-			CompletionDelegate.BindLambda([&](const FName&, UPackage*, EAsyncLoadingResult::Type) 
-			{
-				// Do this open so we can check how many retries occurred.
-				AutoRTFM::Open([&] { NumCompletionCallbacks++; });
-			});
-
-			TArray<int32> RequestIds;
-			RequestIds.Add(LoadPackageAsync(FString::Printf(TEXT("/AutoRTFMTestPackage%d"), __LINE__), CompletionDelegate));
-			RequestIds.Add(LoadPackageAsync(FString::Printf(TEXT("/AutoRTFMTestPackage%d"), __LINE__), CompletionDelegate));
-			RequestIds.Add(LoadPackageAsync(FString::Printf(TEXT("/AutoRTFMTestPackage%d"), __LINE__), CompletionDelegate));
-
-			FlushAsyncLoading(RequestIds);
-			REQUIRE(3 == NumCompletionCallbacks);
-
-			AutoRTFM::AbortTransaction();
-		});
-	}
-
-	SECTION("Stack Local Linker Instancing Context")
-	{
-		int32 RequestId = -1;
-
-		AutoRTFM::Testing::Commit([&]
-		{
-			FLinkerInstancingContext Context;
-			FLoadPackageAsyncOptionalParams Params;
-			Params.InstancingContext = &Context;
-			FString Name(FString::Printf(TEXT("/AutoRTFMTestPackage%d"), __LINE__));
-			RequestId = LoadPackageAsync(Name, MoveTemp(Params));
-		});
-
-		FlushAsyncLoading(RequestId);
-	}
-}
-
 TEST_CASE("UECore.CoreRedirects")
 {
 	FCoreRedirects::Initialize();
diff --git a/Engine/Source/Programs/AutoRTFMTests/Private/UPackageTests.cpp b/Engine/Source/Programs/AutoRTFMTests/Private/UPackageTests.cpp
--- a/Engine/Source/Programs/AutoRTFMTests/Private/UPackageTests.cpp
+++ b/Engine/Source/Programs/AutoRTFMTests/Private/UPackageTests.cpp
@@ -1,8 +1,13 @@
 // Copyright Epic Games, Inc. All Rights Reserved.
 
 #include "AutoRTFM.h"
+#include "AutoRTFMTestUtils.h"
 #include "AutoRTFMTesting.h"
+#include "MyAutoRTFMTestObject.h"
+#include "Misc/PackageName.h"
+#include "UObject/LinkerInstancingContext.h"
 #include "UObject/Package.h"
+#include "UObject/SavePackage.h"
 
 #include "Catch2Includes.h"
 
@@ -97,4 +102,356 @@ TEST_CASE("UPackage.ClearPackageFlags")
 
 		REQUIRE(Package->GetPackageFlags() == (PKG_RuntimeGenerated | PKG_TransientFlags));
 	}
-}
\ No newline at end of file
+}
+
+// Based on `Engine\Private\Tests\Loading\AsyncLoadingTests_Shared.h`, we use similar logic
+// here to make a package that the loader will see and be able to actually load!
+struct FPackageScopedMaker final
+{
+	FString PackagePath;
+
+	FPackageScopedMaker(FString PackageName) : PackagePath(FPackageName::LongPackageNameToFilename(*PackageName, FPackageName::GetAssetPackageExtension()))
+	{
+		// We need to remove any previous package of the same name (could have occurred if a previous test ran segfaulted for instance).
+		if (FPackageName::DoesPackageExist(PackageName))
+		{
+			REQUIRE(IPlatformFile::GetPlatformPhysical().SetReadOnly(*PackagePath, false));
+			REQUIRE(IPlatformFile::GetPlatformPhysical().DeleteFile(*PackagePath));
+		}
+
+		// Ensure that async loading is done.
+		FlushAsyncLoading();
+
+		// Create a package.
+		UPackage* const Package = CreatePackage(*PackageName);
+		
+		// With at least one object in it.
+		FString ObjectName(PackageName);
+		ObjectName.Append(".TestObject");
+		UObject* const Object = NewObject<UMyAutoRTFMTestObject>(Package, *ObjectName, RF_Public | RF_Standalone);
+
+		// Need to mark it is loaded.
+		Package->MarkAsFullyLoaded();
+
+		// Then wipe the standalone flag for reasons.
+		Object->ClearFlags(RF_Standalone);
+
+		// Save the package to the file-system.
+		REQUIRE(UPackage::SavePackage(Package, nullptr, *PackagePath, FSavePackageArgs()));
+
+		// Make sure the package existed in our tables before.
+		REQUIRE(FindObject<UObject>(nullptr, *PackageName) != nullptr);
+
+		// GC and make sure everything gets cleaned up before loading.
+		CollectGarbage(GARBAGE_COLLECTION_KEEPFLAGS);
+
+		// Then make sure the package is no longer loaded in our tables after.
+		REQUIRE(FindObject<UObject>(nullptr, *PackageName) == nullptr);
+	}
+};
+
+TEST_CASE("UPackage.AsyncLoading")
+{
+	SECTION("DoesPackageExist")
+	{
+		AutoRTFM::Testing::Commit([&]
+		{
+			FString Name(FString::Printf(TEXT("/Game/%dAutoRTFMTestPackage%d"), FPlatformProcess::GetCurrentProcessId(), __LINE__));
+			REQUIRE(!FPackageName::DoesPackageExist(Name));
+		});
+		
+		FString Name(FString::Printf(TEXT("/Game/%dAutoRTFMTestPackage%d"), FPlatformProcess::GetCurrentProcessId(), __LINE__));
+		FPackageScopedMaker _(Name);
+
+		AutoRTFM::Testing::Commit([&]
+		{
+			REQUIRE(FPackageName::DoesPackageExist(Name));
+		});
+	}
+
+	SECTION("LoadPackageAsync")
+	{
+		int32 RequestId = -1;
+
+		AutoRTFM::Testing::Commit([&]
+		{
+			FString Name(FString::Printf(TEXT("/AutoRTFMTestPackage%d"), __LINE__));
+			RequestId = LoadPackageAsync(Name);
+		});
+
+		FlushAsyncLoading(RequestId);
+	}
+
+	SECTION("IsAsyncLoading")
+	{
+		int32 RequestId = -1;
+
+		AutoRTFM::Testing::Commit([&]
+		{
+			REQUIRE(!IsAsyncLoading());
+
+			FString Name(FString::Printf(TEXT("/AutoRTFMTestPackage%d"), __LINE__));
+			RequestId = LoadPackageAsync(Name);
+
+			REQUIRE(IsAsyncLoading());
+		});
+
+		FlushAsyncLoading(RequestId);
+	}
+
+	SECTION("FlushAsyncLoading")
+	{
+		FString Name(FString::Printf(TEXT("/AutoRTFMTestPackage%d"), __LINE__));
+		int32 RequestId = LoadPackageAsync(Name);
+
+		AutoRTFM::Testing::Commit([&]
+		{
+			FlushAsyncLoading(RequestId);
+		});
+	}
+
+	SECTION("FlushAsyncLoading Empty")
+	{
+		FString Name(FString::Printf(TEXT("/AutoRTFMTestPackage%d"), __LINE__));
+		int32 RequestId = LoadPackageAsync(Name);
+
+		AutoRTFM::Testing::Commit([&]
+		{
+			FlushAsyncLoading();
+		});
+	}
+
+	SECTION("FlushAsyncLoading One In One Out")
+	{
+		FString Name(FString::Printf(TEXT("/AutoRTFMTestPackage%d"), __LINE__));
+		int32 RequestId1 = LoadPackageAsync(Name);
+
+		AutoRTFM::Testing::Commit([&]
+		{
+			FString Name(FString::Printf(TEXT("/AutoRTFMTestPackage%d"), __LINE__));
+			int32 RequestId2 = LoadPackageAsync(Name);
+			TArray<int32> RequestIds;
+			RequestIds.Add(RequestId1);
+			RequestIds.Add(RequestId2);
+			FlushAsyncLoading(RequestIds);
+		});
+	}
+
+	SECTION("CompletionDelegate is called closed")
+	{
+		AutoRTFM::Testing::Abort([&]
+		{
+			FString Name(FString::Printf(TEXT("/AutoRTFMTestPackage%d"), __LINE__));
+			FLoadPackageAsyncDelegate CompletionDelegate;
+			CompletionDelegate.BindLambda([&](const FName&, UPackage*, EAsyncLoadingResult::Type) { REQUIRE(AutoRTFM::IsClosed()); });
+			int RequestId = LoadPackageAsync(Name, CompletionDelegate);
+			FlushAsyncLoading(RequestId);
+			AutoRTFM::AbortTransaction();
+		});
+	}
+	
+	SECTION("CompletionDelegate aborts")
+	{
+		AutoRTFM::Testing::Abort([&]
+		{
+			FString Name(FString::Printf(TEXT("/AutoRTFMTestPackage%d"), __LINE__));
+			FLoadPackageAsyncDelegate CompletionDelegate;
+			CompletionDelegate.BindLambda([&](const FName&, UPackage*, EAsyncLoadingResult::Type) { AutoRTFM::AbortTransaction(); });
+			int RequestId = LoadPackageAsync(Name, CompletionDelegate);
+			FlushAsyncLoading(RequestId);
+			FAIL("Unreachable!");
+		});
+	}
+	
+	SECTION("FLoadPackageAsyncOptionalParams::CompletionDelegate is called closed")
+	{
+		AutoRTFM::Testing::Abort([&]
+		{
+			FString Name(FString::Printf(TEXT("/AutoRTFMTestPackage%d"), __LINE__));
+			FLoadPackageAsyncOptionalParams Params;
+			Params.CompletionDelegate.Reset(new FLoadPackageAsyncDelegate());
+			Params.CompletionDelegate->BindLambda([&](const FName&, UPackage*, EAsyncLoadingResult::Type) { REQUIRE(AutoRTFM::IsClosed()); });
+			int RequestId = LoadPackageAsync(Name, MoveTemp(Params));
+			FlushAsyncLoading(RequestId);
+			AutoRTFM::AbortTransaction();
+		});
+	}
+	
+	SECTION("FLoadPackageAsyncOptionalParams::CompletionDelegate aborts")
+	{
+		AutoRTFM::Testing::Abort([&]
+		{
+			FString Name(FString::Printf(TEXT("/AutoRTFMTestPackage%d"), __LINE__));
+			FLoadPackageAsyncOptionalParams Params;
+			Params.CompletionDelegate.Reset(new FLoadPackageAsyncDelegate());
+			Params.CompletionDelegate->BindLambda([&](const FName&, UPackage*, EAsyncLoadingResult::Type) { AutoRTFM::AbortTransaction(); });
+			int RequestId = LoadPackageAsync(Name, MoveTemp(Params));
+			FlushAsyncLoading(RequestId);
+			FAIL("Unreachable!");
+		});
+	}
+	
+	SECTION("FLoadPackageAsyncOptionalParams::CompletionDelegate creates UObject")
+	{
+		UMyAutoRTFMTestObject* OpenObject = nullptr;
+		UMyAutoRTFMTestObject* ClosedObject = nullptr;
+
+		AutoRTFM::Testing::Abort([&]
+		{
+			ClosedObject = NewObject<UMyAutoRTFMTestObject>();
+			FString Name(FString::Printf(TEXT("/AutoRTFMTestPackage%d"), __LINE__));
+			FLoadPackageAsyncOptionalParams Params;
+			Params.CompletionDelegate.Reset(new FLoadPackageAsyncDelegate());
+			Params.CompletionDelegate->BindLambda([&](const FName&, UPackage*, EAsyncLoadingResult::Type)
+			{
+				OpenObject = NewObject<UMyAutoRTFMTestObject>();
+				AutoRTFM::AbortTransaction();
+			});
+			int RequestId = LoadPackageAsync(Name, MoveTemp(Params));
+			FlushAsyncLoading(RequestId);
+			FAIL("Unreachable!");
+		});
+
+		REQUIRE(nullptr == ClosedObject);
+		REQUIRE(nullptr == OpenObject);
+	}
+	
+	SECTION("FLoadPackageAsyncOptionalParams::CompletionDelegate calls another LoadPackageAsync")
+	{
+		AutoRTFM::Testing::Abort([&]
+		{
+			FLoadPackageAsyncOptionalParams Params;
+			Params.CompletionDelegate.Reset(new FLoadPackageAsyncDelegate());
+			Params.CompletionDelegate->BindLambda([&](const FName&, UPackage*, EAsyncLoadingResult::Type)
+			{
+				int RequestId = LoadPackageAsync(FString::Printf(TEXT("/AutoRTFMTestPackage%d"), __LINE__), MoveTemp(Params));
+				AutoRTFM::AbortTransaction();
+			});
+			int RequestId = LoadPackageAsync(FString::Printf(TEXT("/AutoRTFMTestPackage%d"), __LINE__), MoveTemp(Params));
+			FlushAsyncLoading(RequestId);
+			FAIL("Unreachable!");
+		});
+	}
+
+	SECTION("Multiple retries because of multiple loads with commit")
+	{
+		AutoRTFMTestUtils::FScopedRetry Retry(AutoRTFM::ForTheRuntime::EAutoRTFMRetryTransactionState::NoRetry);
+
+		int NumCompletionCallbacks = 0;
+		AutoRTFM::Testing::Commit([&]
+		{
+			FLoadPackageAsyncDelegate CompletionDelegate;
+			CompletionDelegate.BindLambda([&](const FName&, UPackage*, EAsyncLoadingResult::Type) 
+			{
+				// Do this open so we can check how many retries occurred.
+				AutoRTFM::Open([&] { NumCompletionCallbacks++; });
+			});
+
+			TArray<int32> RequestIds;
+			RequestIds.Add(LoadPackageAsync(FString::Printf(TEXT("/AutoRTFMTestPackage%d"), __LINE__), CompletionDelegate));
+			RequestIds.Add(LoadPackageAsync(FString::Printf(TEXT("/AutoRTFMTestPackage%d"), __LINE__), CompletionDelegate));
+			RequestIds.Add(LoadPackageAsync(FString::Printf(TEXT("/AutoRTFMTestPackage%d"), __LINE__), CompletionDelegate));
+
+			FlushAsyncLoading(RequestIds);
+			REQUIRE(3 == NumCompletionCallbacks);
+		});
+	}
+
+	SECTION("Multiple retries because of multiple loads with abort")
+	{
+		AutoRTFMTestUtils::FScopedRetry Retry(AutoRTFM::ForTheRuntime::EAutoRTFMRetryTransactionState::NoRetry);
+
+		int NumCompletionCallbacks = 0;
+		AutoRTFM::Testing::Abort([&]
+		{
+			FLoadPackageAsyncDelegate CompletionDelegate;
+			CompletionDelegate.BindLambda([&](const FName&, UPackage*, EAsyncLoadingResult::Type) 
+			{
+				// Do this open so we can check how many retries occurred.
+				AutoRTFM::Open([&] { NumCompletionCallbacks++; });
+			});
+
+			TArray<int32> RequestIds;
+			RequestIds.Add(LoadPackageAsync(FString::Printf(TEXT("/AutoRTFMTestPackage%d"), __LINE__), CompletionDelegate));
+			RequestIds.Add(LoadPackageAsync(FString::Printf(TEXT("/AutoRTFMTestPackage%d"), __LINE__), CompletionDelegate));
+			RequestIds.Add(LoadPackageAsync(FString::Printf(TEXT("/AutoRTFMTestPackage%d"), __LINE__), CompletionDelegate));
+
+			FlushAsyncLoading(RequestIds);
+			REQUIRE(3 == NumCompletionCallbacks);
+
+			AutoRTFM::AbortTransaction();
+		});
+	}
+
+	SECTION("Stack Local Linker Instancing Context")
+	{
+		int32 RequestId = -1;
+
+		AutoRTFM::Testing::Commit([&]
+		{
+			FLinkerInstancingContext Context;
+			FLoadPackageAsyncOptionalParams Params;
+			Params.InstancingContext = &Context;
+			FString Name(FString::Printf(TEXT("/AutoRTFMTestPackage%d"), __LINE__));
+			RequestId = LoadPackageAsync(Name, MoveTemp(Params));
+		});
+
+		FlushAsyncLoading(RequestId);
+	}
+
+	SECTION("Trashed Package")
+	{
+		FString Name(FString::Printf(TEXT("/Game/%dAutoRTFMTestPackage%d"), FPlatformProcess::GetCurrentProcessId(), __LINE__));
+		FPackageScopedMaker _(Name);
+
+		UPackage* const Package = LoadPackage(nullptr, *Name, LOAD_None);
+		REQUIRE(Package);
+			
+		// TODO: If we could move the trashing logic into `UPackage` we could just call that here?
+		const FName NewName = MakeUniqueObjectName(nullptr, UPackage::StaticClass(), NAME_TrashedPackage);
+		Package->Rename(*NewName.ToString(), nullptr, REN_DontCreateRedirectors | REN_NonTransactional | REN_DoNotDirty);
+		Package->SetFlags(RF_Transient);
+
+		AutoRTFM::Testing::Commit([&]
+		{
+			UPackage* const ReloadedPackage = LoadPackage(nullptr, *Name, LOAD_None);
+			REQUIRE(ReloadedPackage);
+			REQUIRE(ReloadedPackage->GetFName() == Name);
+
+			REQUIRE(Package != ReloadedPackage);
+			REQUIRE(!ReloadedPackage->GetFName().ToString().StartsWith(TEXT("TrashedPackage")));
+		});
+	}
+
+	SECTION("Trash Package that was created in same transaction as reloaded")
+	{
+		// #jira SOL-8376: Loading a package, trashing it, then reloading, causes infinite loops because
+		// the transactionally safe async loader will abort and retry for each package load that actually
+		// needs to hit the async loader (because it will fight us for locks).
+		static bool bSOL_8376 = false;
+
+		if (bSOL_8376)
+		{
+			FString Name(FString::Printf(TEXT("/Game/%dAutoRTFMTestPackage%d"), FPlatformProcess::GetCurrentProcessId(), __LINE__));
+			FPackageScopedMaker _(Name);
+
+			AutoRTFM::Testing::Commit([&]
+			{
+				UPackage* const Package = LoadPackage(nullptr, *Name, LOAD_None);
+				REQUIRE(Package);
+			
+				// TODO: If we could move the trashing logic into `UPackage` we could just call that here?
+				const FName NewName = MakeUniqueObjectName(nullptr, UPackage::StaticClass(), NAME_TrashedPackage);
+				Package->Rename(*NewName.ToString(), nullptr, REN_DontCreateRedirectors | REN_NonTransactional | REN_DoNotDirty);
+				Package->SetFlags(RF_Transient);
+
+				UPackage* const ReloadedPackage = LoadPackage(nullptr, *Name, LOAD_None);
+				REQUIRE(ReloadedPackage);
+				REQUIRE(ReloadedPackage->GetFName() == Name);
+
+				REQUIRE(Package != ReloadedPackage);
+				REQUIRE(!ReloadedPackage->GetFName().ToString().StartsWith(TEXT("TrashedPackage")));
+			});
+		}
+	}
+}
diff --git a/Engine/Source/Runtime/CoreUObject/Private/Serialization/TransactionallySafeAsyncLoading.cpp b/Engine/Source/Runtime/CoreUObject/Private/Serialization/TransactionallySafeAsyncLoading.cpp
--- a/Engine/Source/Runtime/CoreUObject/Private/Serialization/TransactionallySafeAsyncLoading.cpp
+++ b/Engine/Source/Runtime/CoreUObject/Private/Serialization/TransactionallySafeAsyncLoading.cpp
@@ -8,6 +8,10 @@
 #include "UObject/Package.h"
 #include "Async/Mutex.h"
 #include "Misc/ScopeLock.h"
+#include "HAL/IConsoleManager.h"
+
+bool CVar_TransactionallySafeAsyncLoadingLegacyTrashedPackages = false;
+static FAutoConsoleVariableRef CVarTransactionallySafeAsyncLoadingLegacyTrashedPackages(TEXT("TransactionallySafeAsyncLoading.LegacyTrashedPackages"), CVar_TransactionallySafeAsyncLoadingLegacyTrashedPackages, TEXT("Use the legacy trashed packages handling (EG. ignore them!)."), ECVF_Default);
 
 // A transactionally safe async package loader that wraps an underlying actual
 // async package loader but allows it to be interacted with safely while also
@@ -462,7 +466,25 @@ private:
 			// became null (meaning the GC did its thing).
 			if (EAsyncLoadingResult::Succeeded == Result)
 			{
-				return LoadedPackage.IsValid();
+				// If the package has been GC'ed then it is not still valid.
+				if (!LoadedPackage.IsValid())
+				{
+					return false;
+				}
+
+				if (UNLIKELY(CVar_TransactionallySafeAsyncLoadingLegacyTrashedPackages))
+				{
+					return true;
+				}
+
+				if (LoadedPackage->GetFName() != PackageName)
+				{
+					static FString TrashedPackage = FName(NAME_TrashedPackage).ToString();
+					check(LoadedPackage->GetFName().ToString().StartsWith(TrashedPackage));
+					return false;
+				}
+
+				return true;
 			}
 			else
 			{
