[AutoRTFM] Fix loader assertion in AutoRTFMTests

The following assertion in Engine\Source\Runtime\CoreUObject\Private\Serialization\PackageStore.cpp was firing:
checkSlow((Backend.Value->GetSupportedLoaders() & OutPackageStoreEntry.LoaderType));
As our fake package store backend handed out FPackageStoreEntrys with the LoaderType = EPackageLoader::LinkerLoad, but EPackageLoader::LinkerLoad was not returned by GetSupportedLoaders().

diff --git a/Engine/Source/Programs/AutoRTFMTests/Private/AutoRTFMTests.cpp b/Engine/Source/Programs/AutoRTFMTests/Private/AutoRTFMTests.cpp
--- a/Engine/Source/Programs/AutoRTFMTests/Private/AutoRTFMTests.cpp
+++ b/Engine/Source/Programs/AutoRTFMTests/Private/AutoRTFMTests.cpp
@@ -139,7 +139,15 @@ int RunTests(int ArgC, const char* ArgV[])
 			return EPackageStoreEntryStatus::Ok;
 		}
 
-		bool GetPackageRedirectInfo(FPackageId PackageId, FName& OutSourcePackageName, FPackageId& OutRedirectedToPackageId) override { return false; }
+		bool GetPackageRedirectInfo(FPackageId PackageId, FName& OutSourcePackageName, FPackageId& OutRedirectedToPackageId) override
+		{
+			return false;
+		}
+		
+		EPackageLoader GetSupportedLoaders() override
+		{
+			return EPackageLoader::LinkerLoad;
+		}
 	};
 
 	TSharedPtr<MyPackageStoreBackend> PackageStoreBackend = MakeShared<MyPackageStoreBackend>();
