Fix for using Stats with AutoRTFM:  Previously a Closed block that could call FStatsThread::AddMessage would undo the StatsMessage additions from an Open block (e.g. FScopeCycleCounter).  The solution is to always ensure the AddMessage is done in the open.

#rnx

diff --git a/Engine/Source/Programs/AutoRTFMTests/Private/AutoRTFMTests.cpp b/Engine/Source/Programs/AutoRTFMTests/Private/AutoRTFMTests.cpp
--- a/Engine/Source/Programs/AutoRTFMTests/Private/AutoRTFMTests.cpp
+++ b/Engine/Source/Programs/AutoRTFMTests/Private/AutoRTFMTests.cpp
@@ -83,8 +83,11 @@ int RunTests(int ArgC, const char* ArgV[])
 		}
 	}
 
-	const TCHAR* CommandLine = TEXT("-Multiprocess -LogCmds=\"LogCsvProfiler off, LogStreaming off, LogUObjectGlobals off, LogPackageName off, LogAutoRTFM warning\" -AsyncLoadingThread");
-	GEngineLoop.PreInit(CommandLine); // Note: Initializes AutoRTFM
+	FString CommandLine = TEXT("-Multiprocess -LogCmds=\"LogCsvProfiler off, LogStreaming off, LogUObjectGlobals off, LogPackageName off, LogAutoRTFM warning\" -AsyncLoadingThread");
+#if FORCE_USE_STATS
+	CommandLine += TEXT(" -LoadTimeStatsForCommandlet");
+#endif
+	GEngineLoop.PreInit(*CommandLine); // Note: Initializes AutoRTFM
 	GLog->SetCurrentThreadAsPrimaryThread();
 
 	if (bRetryNestedToo)
diff --git a/Engine/Source/Programs/AutoRTFMTests/Private/StatsTests.cpp b/Engine/Source/Programs/AutoRTFMTests/Private/StatsTests.cpp
new file mode 100644
--- /dev/null
+++ b/Engine/Source/Programs/AutoRTFMTests/Private/StatsTests.cpp
@@ -0,0 +1,197 @@
+// Copyright Epic Games, Inc. All Rights Reserved.
+
+#include "Catch2Includes.h"
+#include "AutoRTFM.h"
+#include "AutoRTFMTesting.h"
+
+#include "Async/TaskGraphInterfaces.h"
+#include "Stats/Stats.h"
+#include "Stats/StatsCommand.h"
+#include "Stats/StatsData.h"
+#include "Stats/StatsSystem.h"
+#include "Stats/StatsSystemTypes.h"
+
+#if STATS
+
+DECLARE_STATS_GROUP(TEXT("AutoRTFM Tests"), STATGROUP_AutoRTFM_Tests, STATCAT_Advanced);
+DECLARE_CYCLE_STAT(TEXT("Standard Cycle Counter"), STAT_StandardCycleCounter, STATGROUP_AutoRTFM_Tests);
+// Ones we do not have tests for:
+// FScopeCycleCounterEmitter // Requires TemplateObj->SpriteTemplate->GetStatID();
+// FScopeCycleCounterSWidget // Requires SWidget
+// FScopedTraceSolverCallback // Requires ISimCallbackObject
+// FScopeCycleCounterUObject // Requires UObject
+
+// This will put an Add message directly into the Thread Stats
+// Certain stats can perform this internally, we're just simulating that with an existing Stat.
+static void AddDummyStatsMessage()
+{
+	// Keep track of how long we were waiting for it
+	const FName StatName = GET_STATFNAME(STAT_StandardCycleCounter);
+	const uint64 WaitCycles = 0;
+	FThreadStats::AddMessage(StatName, EStatOperation::Add, WaitCycles, true);
+}
+
+// Make sure that the stats advance a frame and that we have its data available on the game thread
+static void TickStatsAndWaitForFrameToBecomeAvailable()
+{
+	// Advance the Stats frame forward.
+	FThreadStats::ExplicitFlush();
+	UE::Stats::FStats::TickCommandletStats();
+
+	// Go through this cycle a bunch of times in case messages are passed back and forth
+	for (int Index = 0; Index < 5; ++Index)
+	{
+		// Flush so that all messages are assigned the last frame and sent to Stats Thread.
+		FThreadStats::ExplicitFlush();
+		// Wait for the Stats Thread to receive the Flush & Process. Sends a message back to Game Thread.
+		FThreadStats::WaitForStats();
+
+		// Now process that message bound for the Game Thread
+		// This includes the update of FLatestGameThreadStatsData 
+		FTaskGraphInterface::Get().ProcessThreadUntilIdle(ENamedThreads::GameThread);
+	}
+}
+
+// Return the accumulated value of the stat (should always be increasing)
+static int64 GetAutoRtfmStatValue()
+{
+	FGameThreadStatsData* ViewDataPtr = FLatestGameThreadStatsData::Get().Latest;
+	if (ViewDataPtr)
+	{
+		static const FName StatName = TEXT("STAT_StandardCycleCounter");
+		if (const FComplexStatMessage* StatMessage = ViewDataPtr->GetStatData(StatName))
+		{
+			return StatMessage->GetValue_Duration(EComplexStatField::IncSum);
+		}
+	}
+
+	return -1.0f;
+}
+
+TEST_CASE("Enable AutoRTFM Stats Gathering")
+{
+	// Try and register FHUDGroupManager::NewFrame but this may not work
+	// Because we may be running with Stats already enabled (and that registration only happens
+	// when we first enable Stats).  Attempt it anyway.
+	UE::Stats::DirectStatsCommand(TEXT("stat AutoRTFM_Tests"));
+
+	// Let's fill the stats history
+	for (int Num = 0; Num < 10; ++Num)
+	{
+		TickStatsAndWaitForFrameToBecomeAvailable();
+	}
+}
+
+TEST_CASE("Stats FScopedCycleCounter Outside Transaction")
+{
+	SCOPE_CYCLE_COUNTER(STAT_StandardCycleCounter);
+}
+
+TEST_CASE("Stats FScopedCycleCounter Inside Transaction")
+{
+	int64 PreviousValue = GetAutoRtfmStatValue();
+	AutoRTFM::Testing::Abort([&]
+	{
+		SCOPE_CYCLE_COUNTER(STAT_StandardCycleCounter);
+		AutoRTFM::AbortTransaction();
+	});
+
+	TickStatsAndWaitForFrameToBecomeAvailable();
+
+	int64 NewValue = GetAutoRtfmStatValue();
+	CATCH_CHECK(NewValue > PreviousValue);
+	PreviousValue = NewValue;
+
+	AutoRTFM::Testing::Commit([&]
+	{
+		SCOPE_CYCLE_COUNTER(STAT_StandardCycleCounter);
+	});
+
+	TickStatsAndWaitForFrameToBecomeAvailable();
+	NewValue = GetAutoRtfmStatValue();
+	CATCH_CHECK(NewValue > PreviousValue);
+}
+
+TEST_CASE("Stats Multi-FScopedCycleCounter Re-Entrant w/ Transaction")
+{
+	auto ChainCall = [](TFunctionRef<void()> ChainedFunc)
+	{
+		SCOPE_CYCLE_COUNTER(STAT_StandardCycleCounter);
+		ChainedFunc();
+	};
+
+	int64 PreviousValue = GetAutoRtfmStatValue();
+	AutoRTFM::Testing::Abort([&]
+	{
+		ChainCall([&]
+		{
+			ChainCall([&]
+			{
+				AutoRTFM::AbortTransaction();
+			});
+		});
+	});
+
+	TickStatsAndWaitForFrameToBecomeAvailable();
+	int64 NewValue = GetAutoRtfmStatValue();
+	CATCH_CHECK(NewValue > PreviousValue);
+	PreviousValue = NewValue;
+
+	AutoRTFM::Testing::Commit([&]
+	{
+		ChainCall([&]
+		{
+			ChainCall([&]
+			{
+			});
+		});
+	});
+
+	TickStatsAndWaitForFrameToBecomeAvailable();
+	NewValue = GetAutoRtfmStatValue();
+	CATCH_CHECK(NewValue > PreviousValue);
+}
+
+TEST_CASE("Stats FScopedCycleCounter w/ Manual Counter")
+{
+	auto ChainCall = [](TFunctionRef<void()> ChainedFunc)
+	{
+		SCOPE_CYCLE_COUNTER(STAT_StandardCycleCounter);
+		ChainedFunc();
+	};
+
+	int64 PreviousValue = GetAutoRtfmStatValue();
+	AutoRTFM::Testing::Abort([&]
+	{
+		ChainCall([&]
+			{
+				ChainCall([&]
+					{
+						AddDummyStatsMessage();
+						AutoRTFM::AbortTransaction();
+					});
+			});
+	});
+
+	TickStatsAndWaitForFrameToBecomeAvailable();
+	int64 NewValue = GetAutoRtfmStatValue();
+	CATCH_CHECK(NewValue > PreviousValue);
+	PreviousValue = NewValue;
+
+	AutoRTFM::Testing::Commit([&]
+	{
+		ChainCall([&]
+			{
+				ChainCall([&]
+					{
+						AddDummyStatsMessage();
+					});
+			});
+	});
+
+	TickStatsAndWaitForFrameToBecomeAvailable();
+	NewValue = GetAutoRtfmStatValue();
+	CATCH_CHECK(NewValue > PreviousValue);
+}
+
+#endif
\ No newline at end of file
diff --git a/Engine/Source/Runtime/Core/Public/Stats/StatsSystemTypes.h b/Engine/Source/Runtime/Core/Public/Stats/StatsSystemTypes.h
--- a/Engine/Source/Runtime/Core/Public/Stats/StatsSystemTypes.h
+++ b/Engine/Source/Runtime/Core/Public/Stats/StatsSystemTypes.h
@@ -1252,6 +1252,7 @@ public:
 	/** Checks the command line whether we want to enable collecting startup stats. */
 	static CORE_API void CheckForCollectingStartupStats();
 
+	UE_AUTORTFM_ALWAYS_OPEN
 	inline void AddStatMessage(const FStatMessage& StatMessage)
 	{
 		LLM_SCOPE(ELLMTag::Stats);
@@ -1287,6 +1288,7 @@ protected:
 	}
 
 	/** Clock operation. **/
+	UE_AUTORTFM_ALWAYS_OPEN
 	inline void AddMessageInner(FName InStatName, EStatOperation::Type InStatOperation)
 	{
 		checkStats((InStatOperation == EStatOperation::CycleScopeStart || InStatOperation == EStatOperation::CycleScopeEnd));
@@ -1529,6 +1531,7 @@ public:
 	 * Pushes the specified stat onto the hierarchy for this thread. Starts
 	 * the timing of the cycles used
 	 */
+	UE_AUTORTFM_ALWAYS_OPEN
 	inline void Start(TStatId InStatId, EStatFlags InStatFlags, bool bAlways = false, UE::FSourceLocation SourceLocation = UE::FSourceLocation::Current())
 	{
 		FMinimalName StatMinimalName = InStatId.GetMinimalName(EMemoryOrder::Relaxed);
@@ -1567,11 +1570,13 @@ public:
 		}
 	}
 
+	UE_AUTORTFM_ALWAYS_OPEN
 	UE_FORCEINLINE_HINT void Start(TStatId InStatId, bool bAlways = false, UE::FSourceLocation SourceLocation = UE::FSourceLocation::Current())
 	{
 		Start(InStatId, EStatFlags::None, bAlways, SourceLocation);
 	}
 
+	UE_AUTORTFM_ALWAYS_OPEN
 	inline void StartTrace(const FName Name)
 	{
 #if CPUPROFILERTRACE_ENABLED
@@ -1583,6 +1588,7 @@ public:
 #endif
 	}
 
+	UE_AUTORTFM_ALWAYS_OPEN
 	inline void StartTrace(const FName Name, const TCHAR* Desc)
 	{
 #if CPUPROFILERTRACE_ENABLED
@@ -1597,6 +1603,7 @@ public:
 	/**
 	 * Stops the capturing and stores the result
 	 */
+	UE_AUTORTFM_ALWAYS_OPEN
 	inline void Stop()
 	{
 		if (EmittedEvent & NamedEvent)
@@ -1622,6 +1629,7 @@ public:
 	/**
 	 * Stops the capturing and stores the result and resets the stat id.
 	 */
+	UE_AUTORTFM_ALWAYS_OPEN
 	inline void StopAndResetStatId()
 	{
 		Stop();
@@ -1643,10 +1651,7 @@ public:
 	 */
 	inline FScopeCycleCounter(TStatId StatId, EStatFlags StatFlags, bool bAlways = false, UE::FSourceLocation SourceLocation = UE::FSourceLocation::Current())
 	{
-		UE_AUTORTFM_OPEN
-		{
-			Start(StatId, StatFlags, bAlways, SourceLocation);
-		};
+		Start(StatId, StatFlags, bAlways, SourceLocation);
 		AutoRTFM::PushOnAbortHandler(this, [this]() { this->Stop(); });
 	}
 
@@ -1661,10 +1666,7 @@ public:
 	inline ~FScopeCycleCounter()
 	{
 		AutoRTFM::PopOnAbortHandler(this);
-		UE_AUTORTFM_OPEN
-		{
-			Stop();
-		};
+		Stop();
 	}
 
 };
