[AutoRTFM] Fix missing throw for AUTORTFM_OPEN / UE_AUTORTFM_ALWAYS_OPEN function attributes

This matches the behaviour of AutoRTFM::Open() / UE_AUTORTFM_OPEN.

diff --git a/Engine/Source/Programs/AutoRTFMTests/Private/Abort.cpp b/Engine/Source/Programs/AutoRTFMTests/Private/Abort.cpp
--- a/Engine/Source/Programs/AutoRTFMTests/Private/Abort.cpp
+++ b/Engine/Source/Programs/AutoRTFMTests/Private/Abort.cpp
@@ -368,6 +368,7 @@ TEST_CASE("Abort.CascadeWithOnCompleteThroughOpen")
 
 			REQUIRE(AutoRTFM::EContextStatus::AbortedByCascadingAbort == Status);
 		});
+		FAIL("AutoRTFM::Open() failed to throw after an abort");
 	});
 
 	REQUIRE(AutoRTFM::ETransactionResult::AbortedByCascade == Result);
@@ -375,6 +376,93 @@ TEST_CASE("Abort.CascadeWithOnCompleteThroughOpen")
 	REQUIRE(bOnCompleteCallbackWasExecuted);
 }
 
+TEST_CASE("Abort.CascadeWithOnCompleteThroughFunctionAttribute")
+{
+	struct FState
+	{
+		AUTORTFM_OPEN void AutoRTFMOpen()
+		{
+			CallCloseTransactCascadingAbort();
+		}
+
+		AUTORTFM_OPEN_NO_VALIDATION void AutoRTFMOpenNoValidation()
+		{
+			CallCloseTransactCascadingAbort();
+		}
+
+		UE_AUTORTFM_ALWAYS_OPEN void UEAutoRTFMAlwaysOpen()
+		{
+			CallCloseTransactCascadingAbort();
+		}
+
+		UE_AUTORTFM_ALWAYS_OPEN_NO_MEMORY_VALIDATION void UEAutoRTFMAlwaysOpenNoMemoryValidation()
+		{
+			CallCloseTransactCascadingAbort();
+		}
+
+		void CallCloseTransactCascadingAbort()
+		{
+			const AutoRTFM::EContextStatus Status = AutoRTFM::Close([&]
+			{
+				AutoRTFM::Transact([&]
+				{
+					AutoRTFM::OnComplete([&]
+					{
+						REQUIRE(!bTouched);  // Rollbacks must occur before the post-abort callback is invoked.
+						bOnCompleteCallbackWasExecuted = true;
+					});
+					AutoRTFM::CascadingAbortTransaction();
+				});
+			});
+		}
+		
+		bool bTouched = false;
+		bool bOnCompleteCallbackWasExecuted = false;
+	};
+	FState State;
+
+	AutoRTFM::ETransactionResult Result = AutoRTFM::ETransactionResult::Committed;
+	SECTION("AUTORTFM_OPEN")
+	{
+		Result = AutoRTFM::Transact([&]
+		{
+			State.bTouched = true;
+			State.AutoRTFMOpen();
+			FAIL("AUTORTFM_OPEN function failed to throw after an abort");
+		});
+	}
+	SECTION("AUTORTFM_OPEN_NO_VALIDATION")
+	{
+		Result = AutoRTFM::Transact([&]
+		{
+			State.bTouched = true;
+			State.AutoRTFMOpenNoValidation();
+			FAIL("AUTORTFM_OPEN_NO_VALIDATION function failed to throw after an abort");
+		});
+	}
+	SECTION("UE_AUTORTFM_ALWAYS_OPEN")
+	{
+		Result = AutoRTFM::Transact([&]
+		{
+			State.bTouched = true;
+			State.UEAutoRTFMAlwaysOpen();
+			FAIL("UE_AUTORTFM_ALWAYS_OPEN function failed to throw after an abort");
+		});
+	}
+	SECTION("UE_AUTORTFM_ALWAYS_OPEN_NO_MEMORY_VALIDATION")
+	{
+		Result = AutoRTFM::Transact([&]
+		{
+			State.bTouched = true;
+			State.UEAutoRTFMAlwaysOpenNoMemoryValidation();
+			FAIL("UE_AUTORTFM_ALWAYS_OPEN_NO_MEMORY_VALIDATION function failed to throw after an abort");
+		});
+	}
+	REQUIRE(AutoRTFM::ETransactionResult::AbortedByCascade == Result);
+	REQUIRE(!State.bTouched);
+	REQUIRE(State.bOnCompleteCallbackWasExecuted);
+}
+
 TEST_CASE("Abort.ForTheRuntime.CascadeThroughManualTransaction")
 {
 	bool bOnCompleteCallbackWasExecuted = false;
diff --git a/Engine/Source/Programs/AutoRTFMTests/Private/OpenAPI.cpp b/Engine/Source/Programs/AutoRTFMTests/Private/OpenAPI.cpp
--- a/Engine/Source/Programs/AutoRTFMTests/Private/OpenAPI.cpp
+++ b/Engine/Source/Programs/AutoRTFMTests/Private/OpenAPI.cpp
@@ -92,7 +92,7 @@ TEST_CASE("OpenAPI.WriteDataInTheOpen")
 	REQUIRE(Value == 1);
 }
 
-TEST_CASE("OpenAPI.RollbackTransaction")
+TEST_CASE("OpenAPI.RollbackTransaction.Lambda")
 {
 	auto TransactResult = AutoRTFM::Transact([&]()
 	{
@@ -105,6 +105,78 @@ TEST_CASE("OpenAPI.RollbackTransaction")
 	REQUIRE(TransactResult == AutoRTFM::ETransactionResult::AbortedByRequest);
 }
 
+TEST_CASE("OpenAPI.RollbackTransaction.AUTORTFM_OPEN")
+{
+	struct S
+	{
+		AUTORTFM_OPEN static void RollbackTransaction()
+		{
+			AutoRTFM::ForTheRuntime::RollbackTransaction();
+		}
+	};
+
+	auto TransactResult = AutoRTFM::Transact([&]()
+	{
+		S::RollbackTransaction();
+		FAIL("AUTORTFM_OPEN failed to throw after an abort");
+	});
+	REQUIRE(TransactResult == AutoRTFM::ETransactionResult::AbortedByRequest);
+}
+
+TEST_CASE("OpenAPI.RollbackTransaction.AUTORTFM_OPEN_NO_VALIDATION")
+{
+	struct S
+	{
+		AUTORTFM_OPEN_NO_VALIDATION static void RollbackTransaction()
+		{
+			AutoRTFM::ForTheRuntime::RollbackTransaction();
+		}
+	};
+
+	auto TransactResult = AutoRTFM::Transact([&]()
+	{
+		S::RollbackTransaction();
+		FAIL("AUTORTFM_OPEN_NO_VALIDATION failed to throw after an abort");
+	});
+	REQUIRE(TransactResult == AutoRTFM::ETransactionResult::AbortedByRequest);
+}
+
+TEST_CASE("OpenAPI.RollbackTransaction.UE_AUTORTFM_ALWAYS_OPEN")
+{
+	struct S
+	{
+		UE_AUTORTFM_ALWAYS_OPEN static void RollbackTransaction()
+		{
+			AutoRTFM::ForTheRuntime::RollbackTransaction();
+		}
+	};
+
+	auto TransactResult = AutoRTFM::Transact([&]()
+	{
+		S::RollbackTransaction();
+		FAIL("UE_AUTORTFM_ALWAYS_OPEN failed to throw after an abort");
+	});
+	REQUIRE(TransactResult == AutoRTFM::ETransactionResult::AbortedByRequest);
+}
+
+TEST_CASE("OpenAPI.RollbackTransaction.UE_AUTORTFM_ALWAYS_OPEN_NO_MEMORY_VALIDATION")
+{
+	struct S
+	{
+		UE_AUTORTFM_ALWAYS_OPEN_NO_MEMORY_VALIDATION static void RollbackTransaction()
+		{
+			AutoRTFM::ForTheRuntime::RollbackTransaction();
+		}
+	};
+
+	auto TransactResult = AutoRTFM::Transact([&]()
+	{
+		S::RollbackTransaction();
+		FAIL("UE_AUTORTFM_ALWAYS_OPEN_NO_MEMORY_VALIDATION failed to throw after an abort");
+	});
+	REQUIRE(TransactResult == AutoRTFM::ETransactionResult::AbortedByRequest);
+}
+
 TEST_CASE("OpenAPI.AbortTransaction")
 {
 	auto TransactResult = AutoRTFM::Transact([&]()
diff --git a/Engine/Source/Runtime/AutoRTFM/Private/API.cpp b/Engine/Source/Runtime/AutoRTFM/Private/API.cpp
--- a/Engine/Source/Runtime/AutoRTFM/Private/API.cpp
+++ b/Engine/Source/Runtime/AutoRTFM/Private/API.cpp
@@ -508,13 +508,18 @@ extern "C" UE_AUTORTFM_NOAUTORTFM UE_AUTORTFM_API void autortfm_pre_open(autortf
 
 extern "C" UE_AUTORTFM_NOAUTORTFM UE_AUTORTFM_API void autortfm_post_open()
 {
-	if (FTransaction* Transaction = FContext::Get()->GetCurrentTransaction())
+	FContext* Context = FContext::Get();
+	if (FTransaction* Transaction = Context->GetCurrentTransaction())
 	{
 		if (Transaction->IsOpenActive())  // Transaction may have been aborted.
 		{
 			Transaction->SetClosedActive();
 		}
 	}
+	if (Context->IsAborting())
+	{
+		Context->Throw();
+	}
 }
 
 extern "C" UE_AUTORTFM_NOAUTORTFM UE_AUTORTFM_API void autortfm_pre_static_local_initializer()
