Get `FCsvProfiler::EndCapture` to work under transactions.

diff --git a/Engine/Source/Programs/AutoRTFMTests/Private/AutoRTFMTests.cpp b/Engine/Source/Programs/AutoRTFMTests/Private/AutoRTFMTests.cpp
--- a/Engine/Source/Programs/AutoRTFMTests/Private/AutoRTFMTests.cpp
+++ b/Engine/Source/Programs/AutoRTFMTests/Private/AutoRTFMTests.cpp
@@ -83,13 +83,15 @@ int RunTests(int ArgC, const char* ArgV[])
 		}
 	}
 
-	FString CommandLine = TEXT("-Multiprocess -LogCmds=\"LogCsvProfiler off, LogStreaming off, LogUObjectGlobals off, LogPackageName off, LogAutoRTFM warning\" -AsyncLoadingThread");
+	FString CommandLine = TEXT("-Multiprocess -csvNoProcessingThread -LogCmds=\"LogCsvProfiler off, LogStreaming off, LogUObjectGlobals off, LogPackageName off, LogAutoRTFM warning\" -AsyncLoadingThread");
 #if FORCE_USE_STATS
 	CommandLine += TEXT(" -LoadTimeStatsForCommandlet");
 #endif
 	GEngineLoop.PreInit(*CommandLine); // Note: Initializes AutoRTFM
 	GLog->SetCurrentThreadAsPrimaryThread();
 
+	FCsvProfiler::Get()->Init();
+
 	if (bRetryNestedToo)
 	{
 		AutoRTFM::ForTheRuntime::SetRetryTransaction(AutoRTFM::ForTheRuntime::EAutoRTFMRetryTransactionState::RetryNestedToo);
diff --git a/Engine/Source/Programs/AutoRTFMTests/Private/CsvProfilerTests.cpp b/Engine/Source/Programs/AutoRTFMTests/Private/CsvProfilerTests.cpp
--- a/Engine/Source/Programs/AutoRTFMTests/Private/CsvProfilerTests.cpp
+++ b/Engine/Source/Programs/AutoRTFMTests/Private/CsvProfilerTests.cpp
@@ -32,8 +32,11 @@ TEST_CASE("CsvProfiler")
 
 		~RAII()
 		{
-			FCsvProfiler::Get()->EndFrame();
 			FCsvProfiler::Get()->EndCapture();
+			FCsvProfiler::Get()->EndFrame();
+
+			// Need to double pump ending frame because of how the profiler works.
+			FCsvProfiler::Get()->EndFrame();
 		}
 	} _;
 
@@ -143,3 +146,34 @@ TEST_CASE("CsvProfiler")
 		FCsvProfiler::SetMetadata(TEXT("SomeKey"), TEXT("SomeValue"));
 	});
 }
+
+TEST_CASE("CsvProfiler.EndCaptureInTransaction")
+{
+	SECTION("No Args")
+	{
+		bool bHit = false;
+
+		FDelegateHandle Handle = FCsvProfiler::Get()->OnCSVProfileEndRequested().AddLambda([&]
+		{
+			REQUIRE(AutoRTFM::IsCommitting());
+			bHit = true;
+		});
+
+		FCsvProfiler::Get()->BeginCapture();
+		FCsvProfiler::Get()->BeginFrame();
+
+		AutoRTFM::Testing::Commit([&]
+		{
+			FCsvProfiler::Get()->EndCapture();
+		});
+
+		FCsvProfiler::Get()->EndFrame();
+
+		// Need to double pump ending frame because of how the profiler works.
+		FCsvProfiler::Get()->EndFrame();
+
+		FCsvProfiler::Get()->OnCSVProfileEndRequested().Remove(Handle);
+
+		REQUIRE(bHit);
+	}
+}
diff --git a/Engine/Source/Runtime/AutoRTFM/Private/API.cpp b/Engine/Source/Runtime/AutoRTFM/Private/API.cpp
--- a/Engine/Source/Runtime/AutoRTFM/Private/API.cpp
+++ b/Engine/Source/Runtime/AutoRTFM/Private/API.cpp
@@ -761,6 +761,19 @@ extern "C" bool autortfm_is_transactional() noexcept
 	return false;
 }
 
+extern "C" UE_AUTORTFM_ALWAYS_OPEN bool autortfm_is_committing() noexcept
+{
+	if (ForTheRuntime::IsAutoRTFMRuntimeEnabled())
+	{
+		if (FContext* Context = FContext::Get())
+		{
+			return Context->IsCommitting();
+		}
+	}
+
+	return false;
+}
+
 extern "C" UE_AUTORTFM_ALWAYS_OPEN bool autortfm_is_committing_or_aborting() noexcept
 {
 	if (ForTheRuntime::IsAutoRTFMRuntimeEnabled())
diff --git a/Engine/Source/Runtime/AutoRTFM/Private/Context.cpp b/Engine/Source/Runtime/AutoRTFM/Private/Context.cpp
--- a/Engine/Source/Runtime/AutoRTFM/Private/Context.cpp
+++ b/Engine/Source/Runtime/AutoRTFM/Private/Context.cpp
@@ -54,6 +54,17 @@ bool FContext::IsTransactional() const
     return GetStatus() == EContextStatus::OnTrack;
 }
 
+bool FContext::IsCommitting() const
+{
+	switch (GetStatus())
+	{
+	default:
+		return false;
+	case EContextStatus::Committing:
+		return true;
+	}
+}
+
 bool FContext::IsCommittingOrAborting() const
 {
 	switch (GetStatus())
diff --git a/Engine/Source/Runtime/AutoRTFM/Private/Context.h b/Engine/Source/Runtime/AutoRTFM/Private/Context.h
--- a/Engine/Source/Runtime/AutoRTFM/Private/Context.h
+++ b/Engine/Source/Runtime/AutoRTFM/Private/Context.h
@@ -45,6 +45,7 @@ public:
 	void ClearTransactionStatus();
 	inline bool IsAborting() const { return IsStatusAborting(Status); }
 	bool IsTransactional() const;
+	bool IsCommitting() const;
 	bool IsCommittingOrAborting() const;
 	bool IsRetrying() const;
 
diff --git a/Engine/Source/Runtime/AutoRTFM/Public/AutoRTFM.h b/Engine/Source/Runtime/AutoRTFM/Public/AutoRTFM.h
--- a/Engine/Source/Runtime/AutoRTFM/Public/AutoRTFM.h
+++ b/Engine/Source/Runtime/AutoRTFM/Public/AutoRTFM.h
@@ -183,6 +183,15 @@ UE_AUTORTFM_CRITICAL_INLINE bool autortfm_is_transactional(void) AUTORTFM_NOEXCE
 }
 #endif
 
+#if UE_AUTORTFM_ENABLED
+UE_AUTORTFM_API bool autortfm_is_committing(void) AUTORTFM_NOEXCEPT;
+#else
+UE_AUTORTFM_CRITICAL_INLINE bool autortfm_is_committing(void) AUTORTFM_NOEXCEPT
+{
+	return false;
+}
+#endif
+
 #if UE_AUTORTFM_ENABLED
 UE_AUTORTFM_API bool autortfm_is_committing_or_aborting(void) AUTORTFM_NOEXCEPT;
 #else
@@ -594,6 +603,10 @@ UE_AUTORTFM_CRITICAL_INLINE_ALWAYS bool IsTransactional() { return autortfm_is_t
 // as true in closed code, and false in open code.
 UE_AUTORTFM_CRITICAL_INLINE_ALWAYS bool IsClosed() { return autortfm_is_closed(); }
 
+// Tells us if we are currently committing a transaction. This will return true
+// when inside an on-commit handler.
+UE_AUTORTFM_CRITICAL_INLINE bool IsCommitting() { return autortfm_is_committing(); }
+
 // Tells us if we are currently committing or aborting a transaction. This will return true
 // when inside an on-abort, on-commit or on-complete handler.
 UE_AUTORTFM_CRITICAL_INLINE bool IsCommittingOrAborting() { return autortfm_is_committing_or_aborting(); }
diff --git a/Engine/Source/Runtime/Core/Private/ProfilingDebugging/CsvProfiler.cpp b/Engine/Source/Runtime/Core/Private/ProfilingDebugging/CsvProfiler.cpp
--- a/Engine/Source/Runtime/Core/Private/ProfilingDebugging/CsvProfiler.cpp
+++ b/Engine/Source/Runtime/Core/Private/ProfilingDebugging/CsvProfiler.cpp
@@ -4024,6 +4024,24 @@ void FCsvProfiler::BeginCapture(int InNumFramesToCapture,
 
 TSharedFuture<FString> FCsvProfiler::EndCapture(FGraphEventRef EventToSignal)
 {
+	if (AutoRTFM::IsClosed())
+	{
+		// In a transaction we cannot support signalling an event.
+		ensure(!EventToSignal.IsValid());
+
+		// If we are in a transaction we only want to actually end the capture if we successfully
+		// committed the transaction. It is safe to do this because the capture API doesn't end
+		// the capture immediately, but at the next end-frame tick (in `EndFrame`).
+		AutoRTFM::OnCommit([this]
+		{
+			EndCapture();
+		});
+
+		// TODO: Could we make a special shared future that would explode if used? We cannot support
+		// shared future returns from within a transaction.
+		return {};
+	}
+
 	LLM_SCOPE(ELLMTag::CsvProfiler);
 	check(IsInGameThread());
 
