[AutoRTFM] Add support for OnPreAbort() callback

This is identical to OnAbort(), but the callback happens before the memory is rolled back.

diff --git a/Engine/Source/Programs/AutoRTFMTests/Private/API.cpp b/Engine/Source/Programs/AutoRTFMTests/Private/API.cpp
--- a/Engine/Source/Programs/AutoRTFMTests/Private/API.cpp
+++ b/Engine/Source/Programs/AutoRTFMTests/Private/API.cpp
@@ -1137,6 +1137,94 @@ TEST_CASE("API.OnAbort")
     REQUIRE(false == bInnerOpenNest);
 }
 
+TEST_CASE("API.OnPreAbort")
+{
+	// Too hard to get this test working when retrying nested transactions so bail!
+	if (AutoRTFM::ForTheRuntime::ShouldRetryNestedTransactionsToo())
+	{
+		return;
+	}
+
+    bool bOuterTransaction = false;
+    bool bInnerTransaction = false;
+    bool bInnerTransactionWithAbort = false;
+    bool bInnerOpenNest = false;
+
+    REQUIRE(AutoRTFM::ETransactionResult::Committed == AutoRTFM::Transact([&]
+    {
+		// If we are retrying transactions, need to reset the test state.
+		AutoRTFM::OnPreAbort([&]
+		{
+			bOuterTransaction = false;
+			bInnerTransaction = false;
+			bInnerTransactionWithAbort = false;
+			bInnerOpenNest = false;
+		});
+
+        AutoRTFM::OnPreAbort([&]
+        {
+            bOuterTransaction = true;
+        });
+
+        // This should only be modified on the commit!
+        if (bOuterTransaction)
+        {
+			AutoRTFM::AbortTransaction();
+        }
+
+        AutoRTFM::Commit([&]
+        {
+            AutoRTFM::OnPreAbort([&]
+            {
+                bInnerTransaction = true;
+            });
+        });
+
+        // This should only be modified on the commit!
+        if (bInnerTransaction)
+        {
+			AutoRTFM::AbortTransaction();
+        }
+
+        AutoRTFM::ETransactionResult Result = AutoRTFM::Transact([&]
+        {
+            AutoRTFM::OnPreAbort([&]
+            {
+                bInnerTransactionWithAbort = true;
+            });
+
+			AutoRTFM::AbortTransaction();
+        });
+
+        REQUIRE(AutoRTFM::ETransactionResult::AbortedByRequest == Result);
+
+        // Inner OnPreAbort runs eagerly
+        if (!bInnerTransactionWithAbort)
+        {
+			AutoRTFM::AbortTransaction();
+        }
+
+        AutoRTFM::Open([&]
+        {
+            AutoRTFM::OnPreAbort([&]
+            {
+                bInnerOpenNest = true;
+            });
+        });
+
+        // This should only be modified on the commit!
+        if (bInnerOpenNest)
+        {
+			AutoRTFM::AbortTransaction();
+        }
+    }));
+
+    REQUIRE(false == bOuterTransaction);
+    REQUIRE(false == bInnerTransaction);
+    REQUIRE(true == bInnerTransactionWithAbort);
+    REQUIRE(false == bInnerOpenNest);
+}
+
 TEST_CASE("API.OnAbort_MutableCapture")
 {
 	FString Message = "Hello";
@@ -1153,6 +1241,22 @@ TEST_CASE("API.OnAbort_MutableCapture")
 	});
 }
 
+TEST_CASE("API.OnPreAbort_MutableCapture")
+{
+	FString Message = "Hello";
+
+	AutoRTFM::Transact([&]
+	{
+		AutoRTFM::OnPreAbort([MessageCopy = Message]() mutable
+		{
+			MessageCopy += " World!";
+			REQUIRE(MessageCopy == "Hello World!");
+		});
+
+		AutoRTFM::AbortTransaction();
+	});
+}
+
 TEST_CASE("API.OnAbortMacro_NoAbort")
 {
 	int Value = 123;
@@ -1170,6 +1274,23 @@ TEST_CASE("API.OnAbortMacro_NoAbort")
 	REQUIRE(Value == 456);
 }
 
+TEST_CASE("API.OnPreAbortMacro_NoAbort")
+{
+	int Value = 123;
+
+	UE_AUTORTFM_TRANSACT
+	{
+		Value = 456;
+
+		UE_AUTORTFM_ONPREABORT(&Value)
+		{
+			Value = 123;
+		};
+	};
+
+	REQUIRE(Value == 456);
+}
+
 TEST_CASE("API.OnAbortMacro_WithAbort")
 {
 	int Value = 123;
@@ -1189,6 +1310,25 @@ TEST_CASE("API.OnAbortMacro_WithAbort")
 	REQUIRE(Value == 123);
 }
 
+TEST_CASE("API.OnPreAbortMacro_WithAbort")
+{
+	int Value = 123;
+
+	UE_AUTORTFM_TRANSACT
+	{
+		Value = 234;
+
+		UE_AUTORTFM_ONPREABORT(&Value)
+		{
+			Value = 123;
+		};
+
+		AutoRTFM::AbortTransaction();
+	};
+
+	REQUIRE(Value == 123);
+}
+
 TEST_CASE("API.DidAllocate")
 {
     constexpr unsigned Size = 1024;
diff --git a/Engine/Source/Programs/AutoRTFMTests/Private/Abort.cpp b/Engine/Source/Programs/AutoRTFMTests/Private/Abort.cpp
--- a/Engine/Source/Programs/AutoRTFMTests/Private/Abort.cpp
+++ b/Engine/Source/Programs/AutoRTFMTests/Private/Abort.cpp
@@ -66,12 +66,24 @@ TEST_CASE("Abort.NestedAbortOrder")
 		AutoRTFM::ETransactionResult Result = AutoRTFM::Transact([&]
 		{
 			AutoRTFM::OnAbort([&]
+				{
+					REQUIRE(3 == Orderer);
+					Orderer += 1;
+				});
+
+			AutoRTFM::OnPreAbort([&]
 				{
 					REQUIRE(1 == Orderer);
 					Orderer += 1;
 				});
 
 			AutoRTFM::OnAbort([&]
+				{
+					REQUIRE(2 == Orderer);
+					Orderer += 1;
+				});
+
+			AutoRTFM::OnPreAbort([&]
 				{
 					REQUIRE(0 == Orderer);
 					Orderer += 1;
@@ -82,7 +94,7 @@ TEST_CASE("Abort.NestedAbortOrder")
 		REQUIRE(AutoRTFM::ETransactionResult::AbortedByRequest == Result);
 	});
 
-	REQUIRE(2 == Orderer);
+	REQUIRE(4 == Orderer);
 }
 
 TEST_CASE("Abort.TransactionInOnCommit")
@@ -128,6 +140,29 @@ TEST_CASE("Abort.TransactionInOnAbort")
 	REQUIRE(AutoRTFM::ETransactionResult::AbortedByRequest == Result);
 }
 
+TEST_CASE("Abort.TransactionInOnPreAbort")
+{
+	AutoRTFM::ETransactionResult Result = AutoRTFM::Transact([&]
+	{
+		AutoRTFM::OnPreAbort([&]
+		{
+			bool bDidSomething = false;
+
+			AutoRTFM::ETransactionResult InnerResult = AutoRTFM::Transact([&]
+			{
+				bDidSomething = true;
+			});
+
+			REQUIRE(AutoRTFM::ETransactionResult::AbortedByTransactDuringAbort == InnerResult);
+			REQUIRE(false == bDidSomething);
+		});
+
+		AutoRTFM::AbortTransaction();
+	});
+
+	REQUIRE(AutoRTFM::ETransactionResult::AbortedByRequest == Result);
+}
+
 TEST_CASE("Abort.TransactionInAbortedOnComplete")
 {
 	AutoRTFM::Testing::Abort([&]
@@ -225,6 +260,35 @@ TEST_CASE("Abort.OnAbortInOuterTransaction_AbortInInnerTransaction")
 	REQUIRE(Value == 4);
 }
 
+TEST_CASE("Abort.OnPreAbortInOuterTransaction_AbortInInnerTransaction")
+{
+	int Value = 1;
+
+	AutoRTFM::Transact([&]
+	{
+		AutoRTFM::OnPreAbort([&]
+		{
+			if (AutoRTFM::ForTheRuntime::GetRetryTransaction() == AutoRTFM::ForTheRuntime::NoRetry)
+			{
+				REQUIRE(!"This should not be reached.");
+				Value = 2;
+			}
+		});
+
+		AutoRTFM::Transact([&]
+		{
+			REQUIRE(Value == 1);
+			Value = 3;
+			AutoRTFM::AbortTransaction();  // Only cancels innermost nest.
+		});
+
+		REQUIRE(Value == 1);
+		Value = 4;
+	});
+
+	REQUIRE(Value == 4);
+}
+
 TEST_CASE("Abort.Cascade")
 {
 	bool bTouched = false;
@@ -272,13 +336,20 @@ TEST_CASE("Abort.CascadeWithOnCompleteCallback")
 
 TEST_CASE("Abort.CascadingAbortMustRunOnAbortsBeforeOnComplete")
 {
+	bool bOnPreAbortWasCalled = false;
 	bool bOnAbortWasCalled = false;
 	bool bPostAbortCallbackWasExecuted = false;
 
 	const AutoRTFM::ETransactionResult Result = AutoRTFM::Transact([&]
 	{
+		AutoRTFM::OnPreAbort([&]
+		{
+			REQUIRE(!bOnAbortWasCalled);  // OnPreAborts must be before OnAborts.
+			bOnPreAbortWasCalled = true;
+		});
 		AutoRTFM::OnAbort([&]
 		{
+			REQUIRE(bOnPreAbortWasCalled);  // OnAborts must be after OnPreAborts.
 			bOnAbortWasCalled = true;
 		});
 		AutoRTFM::OnComplete([&]
@@ -291,6 +362,7 @@ TEST_CASE("Abort.CascadingAbortMustRunOnAbortsBeforeOnComplete")
 
 	REQUIRE(AutoRTFM::ETransactionResult::AbortedByCascade == Result);
 	REQUIRE(bOnAbortWasCalled);
+	REQUIRE(bOnPreAbortWasCalled);
 	REQUIRE(bPostAbortCallbackWasExecuted);
 }
 
@@ -943,11 +1015,18 @@ TEST_CASE("Abort.PushOnAbortHandler_Order")
 
 TEST_CASE("Abort.OnAbortTiming")
 {
+	bool bOnPreAbortRan = false;
 	bool bOnAbortRan = false;
 	int Memory = 666;
 	AutoRTFM::Commit([&]
 	{
 		// If we are retrying transactions, need to reset the test state.
+		AutoRTFM::OnPreAbort([&]
+		{
+			REQUIRE(bOnPreAbortRan);
+			REQUIRE(Memory == 666);
+			bOnPreAbortRan = false;
+		});
 		AutoRTFM::OnAbort([&]
 		{
 			REQUIRE(bOnAbortRan);
@@ -963,9 +1042,16 @@ TEST_CASE("Abort.OnAbortTiming")
 			Memory = 1234;
 			REQUIRE(Memory == 1234);
 
+			AutoRTFM::OnPreAbort([&]
+			{
+				REQUIRE(Memory == 1234);
+				REQUIRE(bOnAbortRan == false);
+				bOnPreAbortRan = true;
+			});
 			AutoRTFM::OnAbort([&]
 			{
 				REQUIRE(Memory == 666);
+				REQUIRE(bOnPreAbortRan == true);
 				bOnAbortRan = true;
 			});
 
@@ -974,6 +1060,7 @@ TEST_CASE("Abort.OnAbortTiming")
 	});
 	REQUIRE(Memory == 666);
 	REQUIRE(bOnAbortRan == true);
+	REQUIRE(bOnPreAbortRan == true);
 }
 
 static void FnHasNoClosed()
@@ -1057,6 +1144,18 @@ TEST_CASE("Abort.StackWriteToOuterOpen")
 		// on abort, then it should change the values of this array.
 		std::array<int, 64> StackGuard{};
 
+		// The OnPreAbort handler should be called *before* the memory is reverted.
+		AutoRTFM::OnPreAbort([&]
+		{
+			if (WritesUndone)
+			{
+				FAIL("OnPreAbort was called after reverting memory");
+			}
+			else if (StackGuard != std::array<int, 64>{})
+			{
+				FAIL("StackGuard was corrupted");
+			}
+		});
 		// The OnAbort handler should be called *after* the memory is reverted.
 		AutoRTFM::OnAbort([&]
 		{
@@ -1158,8 +1257,8 @@ TEST_CASE("Abort.CallbackModifiesTransactionalMemory")
 	AutoRTFMTestUtils::FCaptureWarningContext WarningContext;
 
 	int Memory = 0;
-	
-	SECTION("Task")
+
+	SECTION("OnAbort")
 	{
 		AutoRTFM::Testing::Commit([&]
 		{
@@ -1177,6 +1276,24 @@ TEST_CASE("Abort.CallbackModifiesTransactionalMemory")
 		REQUIRE(WarningContext.HasWarningSubstring(AutoRTFMTestUtils::kMemoryModifiedInAbortHandlerWarning));
 	}
 
+	SECTION("OnPreAbort")
+	{
+		AutoRTFM::Testing::Commit([&]
+		{
+			Memory = 10;
+			AutoRTFM::Testing::Abort([&]
+			{
+				AutoRTFM::OnPreAbort([&]
+				{
+					Memory = 20;
+				});
+				AutoRTFM::AbortTransaction();
+			});
+		});
+
+		REQUIRE(WarningContext.HasWarningSubstring(AutoRTFMTestUtils::kMemoryModifiedInAbortHandlerWarning));
+	}
+
 	SECTION("Destructor")
 	{
 		struct IncrementOnDtor
@@ -1184,7 +1301,7 @@ TEST_CASE("Abort.CallbackModifiesTransactionalMemory")
 			~IncrementOnDtor() { (*Ptr)++; }
 			int* const Ptr;
 		};
-		
+
 		SECTION("OnAbort")
 		{
 			AutoRTFM::Testing::Commit([&]
@@ -1202,6 +1319,23 @@ TEST_CASE("Abort.CallbackModifiesTransactionalMemory")
 			REQUIRE(WarningContext.HasWarningSubstring(AutoRTFMTestUtils::kMemoryModifiedInAbortHandlerDtorWarning));
 		}
 
+		SECTION("OnPreAbort")
+		{
+			AutoRTFM::Testing::Commit([&]
+			{
+				Memory = 10;
+				AutoRTFM::Testing::Abort([&]
+				{
+					AutoRTFM::OnPreAbort([Dtor = IncrementOnDtor{&Memory}]
+					{
+					});
+					AutoRTFM::AbortTransaction();
+				});
+			});
+
+			REQUIRE(WarningContext.HasWarningSubstring(AutoRTFMTestUtils::kMemoryModifiedInAbortHandlerDtorWarning));
+		}
+
 		SECTION("OnCommit")
 		{
 			AutoRTFM::Testing::Commit([&]
diff --git a/Engine/Source/Runtime/AutoRTFM/Private/API.cpp b/Engine/Source/Runtime/AutoRTFM/Private/API.cpp
--- a/Engine/Source/Runtime/AutoRTFM/Private/API.cpp
+++ b/Engine/Source/Runtime/AutoRTFM/Private/API.cpp
@@ -612,6 +612,13 @@ AUTORTFM_INTERNAL void RTFM_OnCommitInternal(TTask<void()> && Work)
 	Context->GetCurrentTransaction()->DeferUntilCommit(std::move(Work));
 }
 
+AUTORTFM_INTERNAL void RTFM_OnPreAbortInternal(TTask<void()>&& Work)
+{
+	FContext* Context = FContext::Get();
+	AUTORTFM_ASSERT(Context->GetStatus() == EContextStatus::OnTrack);
+	Context->GetCurrentTransaction()->DeferUntilPreAbort(std::move(Work));
+}
+
 AUTORTFM_INTERNAL void RTFM_OnAbortInternal(TTask<void()>&& Work)
 {
 	FContext* Context = FContext::Get();
@@ -673,6 +680,11 @@ AUTORTFM_INTERNAL void RTFM_autortfm_on_commit(void (*Work)(void*), void* Arg)
 	RTFM_OnCommitInternal([Work, Arg] { Work(Arg); });
 }
 
+AUTORTFM_INTERNAL void RTFM_autortfm_on_pre_abort(void (*Work)(void*), void* Arg)
+{
+	RTFM_OnPreAbortInternal([Work, Arg] { Work(Arg); });
+}
+
 AUTORTFM_INTERNAL void RTFM_autortfm_on_abort(void (*Work)(void*), void* Arg)
 {
 	RTFM_OnAbortInternal([Work, Arg] { Work(Arg); });
@@ -1090,6 +1102,13 @@ void ForTheRuntime::OnCommitInternal(TTask<void()> && Work)
 	Work();
 }
 
+void ForTheRuntime::OnPreAbortInternal(TTask<void()> && Work)
+{
+	if (autortfm_is_closed())
+	{
+		return RTFM_OnPreAbortInternal(std::move(Work));
+	}
+}
 void ForTheRuntime::OnAbortInternal(TTask<void()> && Work)
 {
 	if (autortfm_is_closed())
@@ -1218,6 +1237,14 @@ extern "C" void autortfm_on_commit(void (*Work)(void*), void* Arg)
     Work(Arg);
 }
 
+extern "C" void autortfm_on_pre_abort(void (*Work)(void*), void* Arg) noexcept
+{
+	if (autortfm_is_closed())
+	{
+		return RTFM_autortfm_on_pre_abort(Work, Arg);
+	}
+}
+
 extern "C" void autortfm_on_abort(void (*Work)(void*), void* Arg) noexcept
 {
 	if (autortfm_is_closed())
diff --git a/Engine/Source/Runtime/AutoRTFM/Private/Stats.h b/Engine/Source/Runtime/AutoRTFM/Private/Stats.h
--- a/Engine/Source/Runtime/AutoRTFM/Private/Stats.h
+++ b/Engine/Source/Runtime/AutoRTFM/Private/Stats.h
@@ -28,6 +28,8 @@ enum class EStatsKind : uint8_t
 	HitSetSkippedBecauseOfStackLocalMemory,
 	AverageCommitTasks,
 	MaximumCommitTasks,
+	AveragePreAbortTasks,
+	MaximumPreAbortTasks,
 	AverageAbortTasks,
 	MaximumAbortTasks,
 	NewMemoryTrackerHit,
@@ -69,6 +71,8 @@ template<EStatsKind Kind> constexpr bool IsStatsKindOneArg()
 	case EStatsKind::MaximumWriteLogBytes:
 	case EStatsKind::AverageCommitTasks:
 	case EStatsKind::MaximumCommitTasks:
+	case EStatsKind::AveragePreAbortTasks:
+	case EStatsKind::MaximumPreAbortTasks:
 	case EStatsKind::AverageAbortTasks:
 	case EStatsKind::MaximumAbortTasks:
 	case EStatsKind::AverageHitSetSize:
@@ -87,6 +91,7 @@ template<EStatsKind Kind> constexpr bool IsStatsKindAverage()
 	case EStatsKind::AverageWriteLogEntries:
 	case EStatsKind::AverageWriteLogBytes:
 	case EStatsKind::AverageCommitTasks:
+	case EStatsKind::AveragePreAbortTasks:
 	case EStatsKind::AverageAbortTasks:
 	case EStatsKind::AverageHitSetSize:
 	case EStatsKind::AverageHitSetCapacity:
diff --git a/Engine/Source/Runtime/AutoRTFM/Private/Transaction.cpp b/Engine/Source/Runtime/AutoRTFM/Private/Transaction.cpp
--- a/Engine/Source/Runtime/AutoRTFM/Private/Transaction.cpp
+++ b/Engine/Source/Runtime/AutoRTFM/Private/Transaction.cpp
@@ -45,6 +45,7 @@ FTransaction** FTransaction::GetIntrusiveAddress()
 FTransaction::FTransaction(FContext* InContext)
     : Context(InContext)
 	, CommitTasks(Context->GetTaskPool())
+	, PreAbortTasks(Context->GetTaskPool())
 	, AbortTasks(Context->GetTaskPool())
 {
 	Resurrect(Context);
@@ -83,6 +84,7 @@ void FTransaction::Reset()
 	AUTORTFM_ASSERT(IsDone());
 
 	CommitTasks.Reset();
+	PreAbortTasks.Reset();
 	AbortTasks.Reset();
 	if (CompletionTasks != std::nullopt)
 	{
@@ -112,6 +114,7 @@ bool FTransaction::IsFresh() const
         && NewMemoryTracker.IsEmpty()
         && WriteLog.IsEmpty()
         && CommitTasks.IsEmpty()
+        && PreAbortTasks.IsEmpty()
         && AbortTasks.IsEmpty()
 		&& (CompletionTasks == std::nullopt || CompletionTasks->IsEmpty())
         && !IsDone()
@@ -149,10 +152,23 @@ void FTransaction::AbortWithoutThrowing()
 		ResetCommitTasksWithValidation();
 	}
 
+	// Call each of the pre memory rollback on-abort handlers.
+	if (AUTORTFM_LIKELY(MemoryValidationLevel == EMemoryValidationLevel::Disabled))
+	{
+		PreAbortTasks.RemoveEachBackward([&](TTask<void()>& Task)
+		{ 
+			Task();
+		});
+	}
+	else
+	{
+		CallAbortTasksWithValidation(PreAbortTasks);
+	}
+
 	// Revert all memory in the write-log.
 	Undo();
 
-	// Call each of the on-abort handlers.
+	// Call each of the post memory rollback on-abort handlers.
 	if (AUTORTFM_LIKELY(MemoryValidationLevel == EMemoryValidationLevel::Disabled))
 	{
 		AbortTasks.RemoveEachBackward([&](TTask<void()>& Task)
@@ -162,7 +178,7 @@ void FTransaction::AbortWithoutThrowing()
 	}
 	else
 	{
-		CallAbortTasksWithValidation();
+		CallAbortTasksWithValidation(AbortTasks);
 	}
 
 	if (IsNested())
@@ -219,7 +235,7 @@ void FTransaction::ResetCommitTasksWithValidation()
 	});
 }
 
-void FTransaction::CallAbortTasksWithValidation()
+void FTransaction::CallAbortTasksWithValidation(FTaskArray& Tasks)
 {
 	FOpenHashThrottler& OpenHashThrottler = Context->GetOpenHashThrottler();
 	const bool bThrottlingEnabled = AutoRTFM::ForTheRuntime::GetMemoryValidationThrottlingEnabled();
@@ -251,7 +267,7 @@ void FTransaction::CallAbortTasksWithValidation()
 		}
 	};
 
-	AbortTasks.RemoveEachBackward([&](TTask<void()>& Task)
+	Tasks.RemoveEachBackward([&](TTask<void()>& Task)
 	{
 		void* const FunctionAddress = Task.FunctionAddress();
 
@@ -383,6 +399,7 @@ void FTransaction::CommitNested()
 	DeferredPopAllOnAbortHandlers.Reset();
 
     Parent->CommitTasks.AddAll(std::move(CommitTasks));
+    Parent->PreAbortTasks.AddAll(std::move(PreAbortTasks));
     Parent->AbortTasks.AddAll(std::move(AbortTasks));
 
     Parent->NewMemoryTracker.Merge(NewMemoryTracker);
@@ -396,6 +413,7 @@ bool FTransaction::AttemptToCommitOuterNest()
 	Context->DumpState();
 	AUTORTFM_VERBOSE("Running commit tasks...");
 
+    PreAbortTasks.Reset();
     AbortTasks.Reset();
 
     CommitTasks.RemoveEachForward([] (TTask<void()>& Task)
diff --git a/Engine/Source/Runtime/AutoRTFM/Private/Transaction.h b/Engine/Source/Runtime/AutoRTFM/Private/Transaction.h
--- a/Engine/Source/Runtime/AutoRTFM/Private/Transaction.h
+++ b/Engine/Source/Runtime/AutoRTFM/Private/Transaction.h
@@ -101,6 +101,7 @@ public:
 	inline bool IsScopedTransaction() const { return bIsStackScoped; }
 
 	void DeferUntilCommit(TTask<void()>&&);
+	void DeferUntilPreAbort(TTask<void()>&&);
 	void DeferUntilAbort(TTask<void()>&&);
 	void DeferUntilComplete(TTask<void()>&&);
 	void PushDeferUntilCommitHandler(const void* Key, TTask<void()>&&);
@@ -167,7 +168,7 @@ private:
 	bool AttemptToCommitOuterNest();
 
 	void ResetCommitTasksWithValidation();
-	void CallAbortTasksWithValidation();
+	void CallAbortTasksWithValidation(FTaskArray& Tasks);
 
 	// Calculates the hash of all the memory written during the transaction and stores it to
 	// RecordedWriteHash.
@@ -214,7 +215,10 @@ private:
 	// Commit tasks run on commit in forward order.
 	FTaskArray CommitTasks;
 
-	// Abort tasks run on abort in reverse order.
+	// Abort tasks run on abort (pre memory rollback) in reverse order.
+	FTaskArray PreAbortTasks;
+
+	// Abort tasks run on abort (post memory rollback) in reverse order.
 	FTaskArray AbortTasks;
 
 	// Completion tasks run at the end of the transaction, in forward order.
diff --git a/Engine/Source/Runtime/AutoRTFM/Private/TransactionInlines.h b/Engine/Source/Runtime/AutoRTFM/Private/TransactionInlines.h
--- a/Engine/Source/Runtime/AutoRTFM/Private/TransactionInlines.h
+++ b/Engine/Source/Runtime/AutoRTFM/Private/TransactionInlines.h
@@ -197,6 +197,15 @@ UE_AUTORTFM_FORCEINLINE void FTransaction::DeferUntilCommit(TTask<void()>&& Call
     CommitTasks.Add(std::move(Copy));
 }
 
+UE_AUTORTFM_FORCEINLINE void FTransaction::DeferUntilPreAbort(TTask<void()>&& Callback)
+{
+	// We explicitly must copy the function here because the original was allocated
+	// within a transactional context, and thus the memory is allocating under
+	// transactionalized conditions. By copying, we create an open copy of the callback.
+	TTask<void()> Copy(Callback);
+    PreAbortTasks.Add(std::move(Copy));
+}
+
 UE_AUTORTFM_FORCEINLINE void FTransaction::DeferUntilAbort(TTask<void()>&& Callback)
 {
 	// We explicitly must copy the function here because the original was allocated
@@ -287,6 +296,9 @@ UE_AUTORTFM_FORCEINLINE void FTransaction::CollectStats() const
     Stats.Collect<EStatsKind::AverageCommitTasks>(CommitTasks.Num());
     Stats.Collect<EStatsKind::MaximumCommitTasks>(CommitTasks.Num());
 
+    Stats.Collect<EStatsKind::AveragePreAbortTasks>(PreAbortTasks.Num());
+    Stats.Collect<EStatsKind::MaximumPreAbortTasks>(PreAbortTasks.Num());
+
     Stats.Collect<EStatsKind::AverageAbortTasks>(AbortTasks.Num());
     Stats.Collect<EStatsKind::MaximumAbortTasks>(AbortTasks.Num());
 
diff --git a/Engine/Source/Runtime/AutoRTFM/Public/AutoRTFM.h b/Engine/Source/Runtime/AutoRTFM/Public/AutoRTFM.h
--- a/Engine/Source/Runtime/AutoRTFM/Public/AutoRTFM.h
+++ b/Engine/Source/Runtime/AutoRTFM/Public/AutoRTFM.h
@@ -362,6 +362,16 @@ UE_AUTORTFM_CRITICAL_INLINE void autortfm_on_commit(void (*work)(void* arg), voi
 }
 #endif
 
+#if UE_AUTORTFM_ENABLED
+UE_AUTORTFM_API void autortfm_on_pre_abort(void (*work)(void* arg), void* arg) AUTORTFM_NOEXCEPT;
+#else
+UE_AUTORTFM_CRITICAL_INLINE void autortfm_on_pre_abort(void (*work)(void* arg), void* arg) AUTORTFM_NOEXCEPT
+{
+	UE_AUTORTFM_UNUSED(work);
+	UE_AUTORTFM_UNUSED(arg);
+}
+#endif
+
 #if UE_AUTORTFM_ENABLED
 UE_AUTORTFM_API void autortfm_on_abort(void (*work)(void* arg), void* arg) AUTORTFM_NOEXCEPT;
 #else
@@ -553,6 +563,7 @@ namespace ForTheRuntime
 	UE_AUTORTFM_API void CascadingAbortTransactionInternal();
 	UE_AUTORTFM_API void CascadingRetryTransactionInternal();
 	UE_AUTORTFM_API void OnCommitInternal(TTask<void()>&& Work);
+	UE_AUTORTFM_API void OnPreAbortInternal(TTask<void()>&& Work);
 	UE_AUTORTFM_API void OnAbortInternal(TTask<void()>&& Work);
 	UE_AUTORTFM_API void OnCompleteInternal(TTask<void()>&& Work);
 	UE_AUTORTFM_API void PushOnCommitHandlerInternal(const void* Key, TTask<void()>&& Work);
@@ -939,7 +950,22 @@ template<typename FunctorType> UE_AUTORTFM_CRITICAL_INLINE void OnCommit(Functor
 #endif
 
 #if UE_AUTORTFM
-// Have some work happen when this transaction aborts.
+// Have some work happen when this transaction aborts (before memory rollback).
+// If an abort occurs, the work list is run in LIFO order.
+// If this is called outside a transaction or from an open nest then the work is ignored.
+template<typename FunctorType> UE_AUTORTFM_CRITICAL_INLINE void OnPreAbort(FunctorType&& Work)
+{
+	if (autortfm_is_closed())
+	{
+		ForTheRuntime::OnPreAbortInternal(std::forward<FunctorType>(Work));
+	}
+}
+#else
+template<typename FunctorType> UE_AUTORTFM_CRITICAL_INLINE void OnPreAbort(FunctorType&&) {}
+#endif
+
+#if UE_AUTORTFM
+// Have some work happen when this transaction aborts (after memory rollback).
 // If an abort occurs, the work list is run in LIFO order.
 // If this is called outside a transaction or from an open nest then the work is ignored.
 template<typename FunctorType> UE_AUTORTFM_CRITICAL_INLINE void OnAbort(FunctorType&& Work)
@@ -1432,6 +1458,14 @@ namespace AutoRTFM::Private
 			return AutoRTFM::Close(std::move(F));
 		}
 	};
+	struct FOnPreAbortHelper
+	{
+		template<typename FunctorType>
+		UE_AUTORTFM_FORCEINLINE void operator+(FunctorType&& F)
+		{
+			AutoRTFM::OnPreAbort(std::forward<FunctorType>(F));
+		}
+	};
 	struct FOnAbortHelper
 	{
 		template<typename FunctorType>
@@ -1475,6 +1509,7 @@ namespace AutoRTFM::Private
 #define UE_AUTORTFM_OPEN_IMPL          ::AutoRTFM::Private::FOpenHelper{} + [&]()
 #define UE_AUTORTFM_OPEN_NO_VALIDATION_IMPL ::AutoRTFM::Private::FOpenNoMemoryValidationHelper{} + [&]()
 #define UE_AUTORTFM_CLOSE_IMPL         ::AutoRTFM::Private::FCloseHelper{} + [&]()
+#define UE_AUTORTFM_ONPREABORT_IMPL(...)  ::AutoRTFM::Private::FOnPreAbortHelper{} + [__VA_ARGS__]() mutable
 #define UE_AUTORTFM_ONABORT_IMPL(...)  ::AutoRTFM::Private::FOnAbortHelper{} + [__VA_ARGS__]() mutable
 #define UE_AUTORTFM_ONCOMMIT_IMPL(...) ::AutoRTFM::Private::FOnCommitHelper{} + [__VA_ARGS__]() mutable
 #define UE_AUTORTFM_TRANSACT_IMPL      ::AutoRTFM::Private::FTransactHelper{} + [&]()
@@ -1485,6 +1520,7 @@ namespace AutoRTFM::Private
 #define UE_AUTORTFM_OPEN_IMPL
 #define UE_AUTORTFM_OPEN_NO_VALIDATION_IMPL
 #define UE_AUTORTFM_CLOSE_IMPL
+#define UE_AUTORTFM_ONPREABORT_IMPL(...) while (false)
 #define UE_AUTORTFM_ONABORT_IMPL(...) while (false)
 #define UE_AUTORTFM_ONCOMMIT_IMPL(...)
 #define UE_AUTORTFM_TRANSACT_IMPL
@@ -1506,7 +1542,13 @@ namespace AutoRTFM::Private
 // Calls should be written like this: UE_AUTORTFM_CLOSE { ... code ... };
 #define UE_AUTORTFM_CLOSE UE_AUTORTFM_CLOSE_IMPL
 
-// Runs a block of code if a transaction aborts.
+// Runs a block of code if a transaction aborts (before memory rollback).
+// In non-transactional code paths the block of code will not be executed at all.
+// The macro arguments are the capture specification for the lambda.
+// Calls should be written like this: UE_AUTORTFM_ONPREABORT(=) { ... code ... };
+#define UE_AUTORTFM_ONPREABORT(...)  UE_AUTORTFM_ONPREABORT_IMPL(__VA_ARGS__)
+
+// Runs a block of code if a transaction aborts (after memory rollback).
 // In non-transactional code paths the block of code will not be executed at all.
 // The macro arguments are the capture specification for the lambda.
 // Calls should be written like this: UE_AUTORTFM_ONABORT(=) { ... code ... };
