Include GCObject.h explicitly pending its removal from StrongObjectPtr.h

diff --git a/Engine/Source/Programs/AutoRTFMTests/Private/AutoRTFMTests.cpp b/Engine/Source/Programs/AutoRTFMTests/Private/AutoRTFMTests.cpp
--- a/Engine/Source/Programs/AutoRTFMTests/Private/AutoRTFMTests.cpp
+++ b/Engine/Source/Programs/AutoRTFMTests/Private/AutoRTFMTests.cpp
@@ -8,6 +8,7 @@
 #include "Misc/PackageName.h"
 #include "RequiredProgramMainCPPInclude.h" // required for ue programs
 #include "HAL/MallocLeakDetection.h"
+#include "UObject/GCObject.h"
 #include "Serialization/PackageStore.h"
 
 IMPLEMENT_APPLICATION(AutoRTFMTests, "AutoRTFMTests");
