Get AutoRTFMTests to work with -MergeModules.

diff --git a/Engine/Source/Programs/AutoRTFMTests/AutoRTFMTestsWithMergeModules.Target.cs b/Engine/Source/Programs/AutoRTFMTests/AutoRTFMTestsWithMergeModules.Target.cs
new file mode 100644
--- /dev/null
+++ b/Engine/Source/Programs/AutoRTFMTests/AutoRTFMTestsWithMergeModules.Target.cs
@@ -0,0 +1,15 @@
+// Copyright Epic Games, Inc. All Rights Reserved.
+
+using EpicGames.Core;
+using UnrealBuildTool;
+using System.Collections.Generic;
+using UnrealBuildBase;
+
+[SupportedPlatforms(UnrealPlatformClass.All)]
+public class AutoRTFMTestsWithMergeModulesTarget : AutoRTFMTestsTarget
+{
+	public AutoRTFMTestsWithMergeModulesTarget(TargetInfo Target) : base(Target)
+	{
+		bMergeModules = true;
+	}
+}
diff --git a/Engine/Source/Programs/AutoRTFMTests/Private/AutoRTFMTests.cpp b/Engine/Source/Programs/AutoRTFMTests/Private/AutoRTFMTests.cpp
--- a/Engine/Source/Programs/AutoRTFMTests/Private/AutoRTFMTests.cpp
+++ b/Engine/Source/Programs/AutoRTFMTests/Private/AutoRTFMTests.cpp
@@ -49,6 +49,8 @@ bool CheckNoMemoryLeaks();
 
 int RunTests(int ArgC, const char* ArgV[])
 {
+	FTaskTagScope Scope(ETaskTag::EGameThread);
+
 	{
 		// measure_environment lazily allocates an Environment pointer on first
 		// call, and holds this in a static variable. Pre-allocate this with the
diff --git a/Engine/Source/Runtime/AutoRTFM/Private/Context.h b/Engine/Source/Runtime/AutoRTFM/Private/Context.h
--- a/Engine/Source/Runtime/AutoRTFM/Private/Context.h
+++ b/Engine/Source/Runtime/AutoRTFM/Private/Context.h
@@ -132,7 +132,7 @@ public:
 	inline FTransactionPool& GetTransactionPool() { return TransactionPool; }
 
 private:
-	static FContext* Instance;
+	UE_AUTORTFM_API static FContext* Instance;
 
 	FContext() { Reset(); }
     FContext(const FContext&) = delete;
diff --git a/Engine/Source/Runtime/AutoRTFM/Private/ExternAPI.h b/Engine/Source/Runtime/AutoRTFM/Private/ExternAPI.h
--- a/Engine/Source/Runtime/AutoRTFM/Private/ExternAPI.h
+++ b/Engine/Source/Runtime/AutoRTFM/Private/ExternAPI.h
@@ -9,7 +9,7 @@
 namespace AutoRTFM
 {
 
-extern autortfm_extern_api GExternAPI;
+UE_AUTORTFM_API extern autortfm_extern_api GExternAPI;
 
 inline void* Allocate(size_t Size, size_t Alignment)
 {
diff --git a/Engine/Source/Runtime/AutoRTFM/Private/ThreadID.h b/Engine/Source/Runtime/AutoRTFM/Private/ThreadID.h
--- a/Engine/Source/Runtime/AutoRTFM/Private/ThreadID.h
+++ b/Engine/Source/Runtime/AutoRTFM/Private/ThreadID.h
@@ -16,7 +16,7 @@ namespace AutoRTFM
 struct FThreadID 
 {
 	// An invalid thread identifier.
-	static const FThreadID Invalid;
+	UE_AUTORTFM_API static const FThreadID Invalid;
 
 	// Returns the currently executing thread's unique identifier.
 	UE_AUTORTFM_API
diff --git a/Engine/Source/Runtime/Core/Private/Modules/ModuleManager.cpp b/Engine/Source/Runtime/Core/Private/Modules/ModuleManager.cpp
--- a/Engine/Source/Runtime/Core/Private/Modules/ModuleManager.cpp
+++ b/Engine/Source/Runtime/Core/Private/Modules/ModuleManager.cpp
@@ -3,6 +3,7 @@
 #include "Modules/ModuleManager.h"
 
 #include "HAL/FileManager.h"
+#include "HAL/IConsoleManager.h"
 #include "Internationalization/StringTableCore.h"
 #include "Logging/StructuredLog.h"
 #include "Misc/App.h"
