[AutoRTFM] Fix AutoRTFM crash inside the structured logging code.

This code is reachable from Verse so now uses transactionally-safe mutexes.

diff --git a/Engine/Source/Programs/AutoRTFMTests/Private/UECoreTests.cpp b/Engine/Source/Programs/AutoRTFMTests/Private/UECoreTests.cpp
--- a/Engine/Source/Programs/AutoRTFMTests/Private/UECoreTests.cpp
+++ b/Engine/Source/Programs/AutoRTFMTests/Private/UECoreTests.cpp
@@ -18,6 +18,7 @@
 #include "Internationalization/TextHistory.h"
 #include "Logging/LogMacros.h"
 #include "Logging/StructuredLog.h"
+#include "Logging/StructuredLogFormat.h"
 #include "Memory/VirtualStackAllocator.h"
 #include "Modules/ModuleManager.h"
 #include "Serialization/CustomVersion.h"
@@ -989,7 +990,7 @@ TEST_CASE("TStatId")
 				FDynamicStats::CreateStatIdInt64<FStatGroup_STATGROUP_AutoRTFMTesting>(Name);
 			}
 
-		
+
 			{
 				FString Name = "Test"; Name.AppendInt(__LINE__);
 				FDynamicStats::CreateStatIdDouble<FStatGroup_STATGROUP_AutoRTFMTesting>(Name);
@@ -1374,7 +1375,7 @@ TEST_CASE("UECore.TQueue")
 			});
 
 			REQUIRE(!bIsEmpty);
-			
+
 			AutoRTFM::Testing::Commit([&]
 			{
 				bIsEmpty = Queue.IsEmpty();
@@ -1483,12 +1484,12 @@ TEST_CASE("UECore.FConfigFile")
 		Config.FindOrAddConfigSection(TEXT("WOW"));
 
 		REQUIRE(!Config.IsEmpty());
-	
+
 		AutoRTFM::Testing::Commit([&]
 		{
 			Config.Empty();
 		});
-		
+
 		REQUIRE(Config.IsEmpty());
 	}
 }
@@ -1760,7 +1761,7 @@ TEST_CASE("UECore.CoreRedirects")
 			FCoreRedirects::RemoveRedirectList(Redirects(), TEXT("AutoRTFMTests.UECore.CoreRedirects"));
 			REQUIRE(FCoreRedirects::GetRedirectedName(ECoreRedirectFlags::Type_Package, From) == From);
 		}
-		
+
 		SECTION("Abort")
 		{
 			AutoRTFM::Testing::Abort([&]
@@ -1771,11 +1772,11 @@ TEST_CASE("UECore.CoreRedirects")
 			REQUIRE(FCoreRedirects::GetRedirectedName(ECoreRedirectFlags::Type_Package, From) == From);
 		}
 	}
-	
+
 	SECTION("RemoveRedirectList")
 	{
 		FCoreRedirects::AddRedirectList(Redirects(), TEXT("AutoRTFMTests.UECore.CoreRedirects"));
-		
+
 		SECTION("Commit")
 		{
 			AutoRTFM::Testing::Commit([&]
@@ -1784,7 +1785,7 @@ TEST_CASE("UECore.CoreRedirects")
 			});
 			REQUIRE(FCoreRedirects::GetRedirectedName(ECoreRedirectFlags::Type_Package, From) == From);
 		}
-		
+
 		SECTION("Abort")
 		{
 			AutoRTFM::Testing::Abort([&]
@@ -1833,15 +1834,15 @@ public:
 	bool Truncate(int64 NewSize) override										{ unimplemented(); return false; }
 	bool Flush(const bool bFullFlush = false) override							{ return true; }
 	void ShrinkBuffers() override												{}
-	
+
 	int64 Tell() override
-	{ 
-		return Cursor; 
+	{
+		return Cursor;
 	}
 
 	int64 Size() override
-	{ 
-		return strlen(Data); 
+	{
+		return strlen(Data);
 	}
 
 	bool Seek(int64 NewPosition) override
@@ -1890,7 +1891,7 @@ public:
 	bool IterateDirectory(const TCHAR* Directory, FDirectoryVisitor& Visitor) override                    { unimplemented(); return false; }
 	bool IterateDirectoryStat(const TCHAR* Directory, FDirectoryStatVisitor& Visitor) override            { unimplemented(); return false; }
 
-	const TCHAR* GetName() const override 
+	const TCHAR* GetName() const override
 	{
 		return TEXT("FakePlatformFile");
 	}
@@ -2033,12 +2034,12 @@ TEST_CASE("BlueprintExceptionInfo.Tracepoint")
 {
 	AutoRTFM::Testing::Commit([&]
 	{
-		// We should be able to construct millions of tracepoints within a transaction without causing 
+		// We should be able to construct millions of tracepoints within a transaction without causing
 		// the task array to overflow.
 		for (int Count = 0; Count < 50'000'000; ++Count)
 		{
 			FBlueprintExceptionInfo TracepointExceptionInfo(EBlueprintExceptionType::Tracepoint);
-			
+
 			// REQUIRE is actually too slow here, as it goes into the open and back.
 			check(TracepointExceptionInfo.GetType() == EBlueprintExceptionType::Tracepoint);
 		}
@@ -2047,7 +2048,7 @@ TEST_CASE("BlueprintExceptionInfo.Tracepoint")
 
 TEST_CASE("FVirtualStackAllocator.NestedFrames")
 {
-	// This test case is loosely adapted from "Testing FVirtualStackAllocator ThreadSingleton and Macros" 
+	// This test case is loosely adapted from "Testing FVirtualStackAllocator ThreadSingleton and Macros"
 	// in VirtualStackAllocatorTests.cpp.
 	FVirtualStackAllocator Allocator(32768, EVirtualStackAllocatorDecommitMode::ExcessOnStackEmpty);
 
@@ -2091,7 +2092,7 @@ TEST_CASE("FVirtualStackAllocator.NestedFrames")
 		{
 			UE_VSTACK_MAKE_FRAME(NestedBookmark, &Allocator);
 			UE_VSTACK_ALLOC(&Allocator, 256);
-	
+
 			AutoRTFM::AbortTransaction();
 		});
 
@@ -2126,7 +2127,7 @@ TEST_CASE("FVirtualStackAllocator.PreventsOOM")
 	// `AllOnDestruction` is the decommit mode used by FBlueprintContext.
 	FVirtualStackAllocator Allocator(8*1024*1024, EVirtualStackAllocatorDecommitMode::AllOnDestruction);
 
-	// This test case takes a while with the memory validator on, because it wants to do a validation 
+	// This test case takes a while with the memory validator on, because it wants to do a validation
 	// check on every iteration through the loop. So we temporarily disable validation for efficiency.
 	AutoRTFM::EMemoryValidationLevel OriginalLevel = AutoRTFM::ForTheRuntime::GetMemoryValidationLevel();
 	AutoRTFM::ForTheRuntime::SetMemoryValidationLevel(AutoRTFM::EMemoryValidationLevel::Disabled);
@@ -2332,3 +2333,22 @@ TEST_CASE("FDynamicallyTypedValue")
 	}
 }
 
+TEST_CASE("UECore.FInlineLogTemplate")
+{
+	// Verify that we can instantiate localized log templates from within a transaction.
+	SECTION("Commit")
+	{
+		AutoRTFM::Testing::Commit([&]
+		{
+			UE::FInlineLogTemplate LogTempl(FText::FromString("Cat"));
+		});
+	}
+	SECTION("Abort")
+	{
+		AutoRTFM::Testing::Abort([&]
+		{
+			UE::FInlineLogTemplate LogTempl(FText::FromString("Dog"));
+			AutoRTFM::AbortTransaction();
+		});
+	}
+}
\ No newline at end of file
diff --git a/Engine/Source/Runtime/Core/Private/Logging/StructuredLog.cpp b/Engine/Source/Runtime/Core/Private/Logging/StructuredLog.cpp
--- a/Engine/Source/Runtime/Core/Private/Logging/StructuredLog.cpp
+++ b/Engine/Source/Runtime/Core/Private/Logging/StructuredLog.cpp
@@ -7,6 +7,7 @@
 #include "Algo/Find.h"
 #include "Algo/NoneOf.h"
 #include "Async/Mutex.h"
+#include "Async/TransactionallySafeMutex.h"
 #include "Async/UniqueLock.h"
 #include "AutoRTFM.h"
 #include "Containers/AnsiString.h"
@@ -384,7 +385,7 @@ public:
 
 private:
 	TSparseArray<FTextFormat> Table;
-	mutable FMutex Mutex;
+	mutable FTransactionallySafeMutex Mutex;
 };
 
 static FLocalizedLogFormatTable& GetLocalizedLogFormatTable()
