Add an additional async loading test that creates a UObject before issuing a sync load (which requires heroics in the transactionally-safe async loader to work around).

diff --git a/Engine/Source/Programs/AutoRTFMTests/Private/UPackageTests.cpp b/Engine/Source/Programs/AutoRTFMTests/Private/UPackageTests.cpp
--- a/Engine/Source/Programs/AutoRTFMTests/Private/UPackageTests.cpp
+++ b/Engine/Source/Programs/AutoRTFMTests/Private/UPackageTests.cpp
@@ -668,4 +668,20 @@ TEST_CASE("UPackage.AsyncLoading")
 
 		REQUIRE(bHit);
 	}
+
+	SECTION("Create unrelated UObject's before package load")
+	{
+		FString Name(FString::Printf(TEXT("/Game/%dAutoRTFMTestPackage%d"), FPlatformProcess::GetCurrentProcessId(), __LINE__));
+		FPackageScopedMaker _(Name);
+
+		AutoRTFM::Testing::Commit([&]
+		{
+			// A completely unrelated UObject to the package.
+			UPackage* const UnrelatedPackage = NewObject<UPackage>();
+			REQUIRE(UnrelatedPackage);
+
+			UPackage* const Package = LoadPackage(nullptr, *Name, LOAD_None);
+			REQUIRE(Package);
+		});
+	}
 }
