[AutoRTFM] Make RWLock hazards as easy as possible to spot in a crash log.

We now report all read-write lock errors with an informative banner.

#FailedValidations Preflight

diff --git a/Engine/Source/Programs/AutoRTFMTests/Private/Abort.cpp b/Engine/Source/Programs/AutoRTFMTests/Private/Abort.cpp
--- a/Engine/Source/Programs/AutoRTFMTests/Private/Abort.cpp
+++ b/Engine/Source/Programs/AutoRTFMTests/Private/Abort.cpp
@@ -7,6 +7,7 @@
 #include "Catch2Includes.h"
 #include "Misc/StringOutputDevice.h"
 #include "Misc/OutputDeviceRedirector.h"
+#include "Logging/LogScopedVerbosityOverride.h"
 
 #include <map>
 #include <vector>
@@ -1366,14 +1367,11 @@ class FScopedErrorHintTest
 	public:
 		FScopedErrorHintTest(FOutputDevice& InDevice) : Device(InDevice)
 		{
-			OriginalVerbosity = LogAutoRTFM.GetVerbosity();
 			GLog->AddOutputDevice(&Device);
-			LogAutoRTFM.SetVerbosity(ELogVerbosity::Display);
 		}
 
 		~FScopedErrorHintTest()
 		{
-			LogAutoRTFM.SetVerbosity(OriginalVerbosity);
 			GLog->RemoveOutputDevice(&Device);
 		}
 
@@ -1381,7 +1379,7 @@ class FScopedErrorHintTest
 		AutoRTFMTestUtils::FScopedInternalAbortAction Scoped1{AutoRTFM::ForTheRuntime::EAutoRTFMInternalAbortActionState::Abort};
 		AutoRTFMTestUtils::FScopedEnsureOnInternalAbort Scoped2{false};
 		AutoRTFMTestUtils::FScopedHazardLogSeverity Scoped3{autortfm_log_info};
-		ELogVerbosity::Type OriginalVerbosity;
+		FLogScopedVerbosityOverride Scoped4{&LogAutoRTFM, ELogVerbosity::Display};
 		FOutputDevice& Device;
 };
 
@@ -1422,3 +1420,39 @@ TEST_CASE("Abort.Language.ErrorHints.EnterCriticalSection")
 	// Did we print our message?
 	REQUIRE(StringLog.Contains("A critical section was accessed while inside an AutoRTFM transaction."));
 }
+
+TEST_CASE("Abort.Language.ErrorHints.AcquireSRWLockExclusive")
+{
+	TLogCategoryOutputDevice StringLog(LogAutoRTFM);
+	FScopedErrorHintTest TestWith(StringLog);
+
+	FRWLock RW;
+
+	const AutoRTFM::ETransactionResult Result = AutoRTFM::Transact([&]
+	{
+		RW.WriteLock();
+	});
+
+	REQUIRE(AutoRTFM::ETransactionResult::AbortedByLanguage == Result);
+
+	// Did we print our message?
+	REQUIRE(StringLog.Contains("A read-write lock was accessed while inside an AutoRTFM transaction."));
+}
+
+TEST_CASE("Abort.Language.ErrorHints.TryAcquireSRWLockShared")
+{
+	TLogCategoryOutputDevice StringLog(LogAutoRTFM);
+	FScopedErrorHintTest TestWith(StringLog);
+
+	FRWLock RW;
+
+	const AutoRTFM::ETransactionResult Result = AutoRTFM::Transact([&]
+	{
+		RW.TryReadLock();
+	});
+
+	REQUIRE(AutoRTFM::ETransactionResult::AbortedByLanguage == Result);
+
+	// Did we print our message?
+	REQUIRE(StringLog.Contains("A read-write lock was accessed while inside an AutoRTFM transaction."));
+}
diff --git a/Engine/Source/Runtime/AutoRTFM/Private/FunctionMap.cpp b/Engine/Source/Runtime/AutoRTFM/Private/FunctionMap.cpp
--- a/Engine/Source/Runtime/AutoRTFM/Private/FunctionMap.cpp
+++ b/Engine/Source/Runtime/AutoRTFM/Private/FunctionMap.cpp
@@ -291,6 +291,21 @@ replace it with FTransactionallySafeCriticalSection. If FScopeLock is used, repl
 	AutoRTFM::LogWithCallstack(Severity, "Call stack:");
 }
 
+AUTORTFM_DISABLE static void ReportAutoRTFMRWLockHazard()
+{
+	const autortfm_log_severity Severity = AutoRTFM::ForTheRuntime::GetAutoRTFMHazardLogSeverity();
+	LogAutoRTFMHazard(Severity, R"(%s
+A read-write lock was accessed while inside an AutoRTFM transaction. AutoRTFM doesn't support native
+read-write locks, but it does support transactionally-safe read-write locks that can be used instead.
+
+Use the call stack below to find the FRWLock accessed immediately before this error occurred, and
+replace it with FTransactionallySafeFRWLock. If FReadScopeLock, FWriteScopeLock or FRWScopeLock are
+used, replace them with UE::TReadScopeLock, UE::TWriteScopeLock or UE::TRWScopeLock.
+)",
+		GetAutoRTFMHazardBanner());
+	AutoRTFM::LogWithCallstack(Severity, "Call stack:");
+}
+
 AUTORTFM_DISABLE static void* FunctionMapReportError(void* Open, const char* Where)
 {
 #if AUTORTFM_PLATFORM_WINDOWS
@@ -304,6 +319,16 @@ AUTORTFM_DISABLE static void* FunctionMapReportError(void* Open, const char* Whe
 	{
 		ReportAutoRTFMCriticalSectionHazard();
 	}
+	else if (Open == &TryAcquireSRWLockExclusive ||
+		Open == &AcquireSRWLockExclusive ||
+		Open == &ReleaseSRWLockExclusive ||
+		Open == &TryAcquireSRWLockShared ||
+		Open == &AcquireSRWLockShared ||
+		Open == &ReleaseSRWLockShared ||
+		Open == &TryAcquireSRWLockExclusive)
+	{
+		ReportAutoRTFMRWLockHazard();
+	}
 #elif AUTORTFM_PLATFORM_LINUX
 	if (Open == &pthread_mutexattr_init ||
 		Open == &pthread_mutexattr_settype ||
@@ -316,6 +341,16 @@ AUTORTFM_DISABLE static void* FunctionMapReportError(void* Open, const char* Whe
 	{
 		ReportAutoRTFMCriticalSectionHazard();
 	}
+	else if (Open == &pthread_rwlock_init ||
+		Open == &pthread_rwlock_destroy ||
+		Open == &pthread_rwlock_trywrlock ||
+		Open == &pthread_rwlock_wrlock ||
+		Open == &pthread_rwlock_unlock ||
+		Open == &pthread_rwlock_tryrdlock ||
+		Open == &pthread_rwlock_rdlock)
+	{
+		ReportAutoRTFMRWLockHazard();
+	}
 #endif
 
 	std::string FunctionDescription = GetFunctionDescription(Open);
