Fix std::string test in Open.ReturnValue.

This started crashing on my machine out of the blue (in Development). The simplest fix is to just use FString instead of std::string; the test still demonstrates the same functionality.

diff --git a/Engine/Source/Programs/AutoRTFMTests/Private/Open.cpp b/Engine/Source/Programs/AutoRTFMTests/Private/Open.cpp
--- a/Engine/Source/Programs/AutoRTFMTests/Private/Open.cpp
+++ b/Engine/Source/Programs/AutoRTFMTests/Private/Open.cpp
@@ -185,23 +185,19 @@ TEST_CASE("Open.ReturnValue")
         REQUIRE(AutoRTFM::ETransactionResult::Committed == Result);
         REQUIRE(42 == Value);
     }
-
-#if !UE_BUILD_DEBUG
-	// jira SOL-7541: std::string only under debug+linux hits a stdlib
-	// missing function in the closed!
-    SECTION("char*")
+	
+    SECTION("TCHAR*")
     {
-        std::string Value = "<unassigned>";
+        FString Value = TEXT("<unassigned>");
         AutoRTFM::ETransactionResult Result = AutoRTFM::Transact([&]
-            {
-                // Note: AutoRTFM::Open() is returning a const char*
-                std::string Got = AutoRTFM::Open([] { return "meow"; });
-                AutoRTFM::Open([&] { Value = Got; });
-            });
+        {
+            // Note: AutoRTFM::Open() is returning a TCHAR*
+            FString Got = AutoRTFM::Open([] { return TEXT("meow"); });
+            AutoRTFM::Open([&] { Value = Got; });
+        });
         REQUIRE(AutoRTFM::ETransactionResult::Committed == Result);
-        REQUIRE("meow" == Value);
+        REQUIRE(Value == TEXT("meow"));
     }
-#endif //!UE_BUILD_DEBUG
 
     SECTION("Custom type")
     {
