[AutoRTFM] Address BugHawk-reported issues.

diff --git a/Engine/Source/Programs/AutoRTFMTests/Private/UECoreTests.cpp b/Engine/Source/Programs/AutoRTFMTests/Private/UECoreTests.cpp
--- a/Engine/Source/Programs/AutoRTFMTests/Private/UECoreTests.cpp
+++ b/Engine/Source/Programs/AutoRTFMTests/Private/UECoreTests.cpp
@@ -1206,6 +1206,7 @@ TEST_CASE("UECore.FString")
 			String = FString::Printf(TEXT("Foo '%s' Bar"), TEXT("Stuff"));
 		});
 
+		// FString::operator== uses case insensitive comparison (Stricmp).
 		REQUIRE(String == "Foo 'Stuff' BAR");
 	}
 
@@ -2343,10 +2344,11 @@ TEST_CASE("UECore.UnifiedError")
 	}
 	SECTION("Abort")
 	{
-		AutoRTFM::Testing::Commit([&]
+		AutoRTFM::Testing::Abort([&]
 		{
 			UE::UnifiedError::FError Err = UE::UnifiedError::Core::CancellationError::MakeError();
 			FText Text = Err.GetErrorMessage(/*bIncludeContext*/ false);
+			AutoRTFM::AbortTransaction();
 		});
 	}
 }
