Make a internal API for trashing objects and use it everywhere we previously had lots of copy-pasta code.

diff --git a/Engine/Source/Programs/AutoRTFMTests/AutoRTFMTests.Build.cs b/Engine/Source/Programs/AutoRTFMTests/AutoRTFMTests.Build.cs
--- a/Engine/Source/Programs/AutoRTFMTests/AutoRTFMTests.Build.cs
+++ b/Engine/Source/Programs/AutoRTFMTests/AutoRTFMTests.Build.cs
@@ -1,5 +1,6 @@
 // Copyright Epic Games, Inc. All Rights Reserved.
 
+using System.IO;
 using UnrealBuildTool;
 
 public class AutoRTFMTests : ModuleRules
@@ -29,8 +30,9 @@ public class AutoRTFMTests : ModuleRules
 		PrivateIncludePaths.AddRange(
 			new string[]
 			{
-				"Runtime/AutoRTFM/Private",
-				"Runtime/Core/Private",
+				Path.Combine(GetModuleDirectory("AutoRTFM"), "Private"),
+				Path.Combine(GetModuleDirectory("Core"), "Private"),
+				Path.Combine(GetModuleDirectory("CoreUObject"), "Internal"),
 			});
 
 		PCHUsage = PCHUsageMode.NoPCHs;
diff --git a/Engine/Source/Programs/AutoRTFMTests/Classes/MyAutoRTFMTestObject.h b/Engine/Source/Programs/AutoRTFMTests/Classes/MyAutoRTFMTestObject.h
--- a/Engine/Source/Programs/AutoRTFMTests/Classes/MyAutoRTFMTestObject.h
+++ b/Engine/Source/Programs/AutoRTFMTests/Classes/MyAutoRTFMTestObject.h
@@ -13,7 +13,7 @@ public:
 
 	UMyAutoRTFMTestObject(const FObjectInitializer& ObjectInitializer = FObjectInitializer::Get());
 
-	int Value;
+	int Value = 42;
 
 	static FConstructorCallback* ConstructorCallback;
 
diff --git a/Engine/Source/Programs/AutoRTFMTests/Private/UPackageTests.cpp b/Engine/Source/Programs/AutoRTFMTests/Private/UPackageTests.cpp
--- a/Engine/Source/Programs/AutoRTFMTests/Private/UPackageTests.cpp
+++ b/Engine/Source/Programs/AutoRTFMTests/Private/UPackageTests.cpp
@@ -8,17 +8,10 @@
 #include "UObject/LinkerInstancingContext.h"
 #include "UObject/Package.h"
 #include "UObject/SavePackage.h"
+#include "UObject/UObjectGlobalsInternal.h"
 
 #include "Catch2Includes.h"
 
-static void TrashPackage(UPackage* const Package)
-{
-	// TODO: If we could move the trashing logic into `UPackage` we could just call that here?
-	const FName NewName = MakeUniqueObjectName(nullptr, UPackage::StaticClass(), NAME_TrashedPackage);
-	Package->Rename(*NewName.ToString(), nullptr, REN_DontCreateRedirectors | REN_NonTransactional | REN_DoNotDirty);
-	Package->SetFlags(RF_Transient);
-}
-
 TEST_CASE("UPackage.SetPackageFlagsTo")
 {
 	SECTION("Commit")
@@ -112,6 +105,8 @@ TEST_CASE("UPackage.ClearPackageFlags")
 	}
 }
 
+static const TCHAR* ObjectName = TEXT("TestObject");
+
 // Based on `Engine\Private\Tests\Loading\AsyncLoadingTests_Shared.h`, we use similar logic
 // here to make a package that the loader will see and be able to actually load!
 struct FPackageScopedMaker final
@@ -133,19 +128,13 @@ struct FPackageScopedMaker final
 		// Create a package.
 		UPackage* const Package = CreatePackage(*PackageName);
 		
-		// With at least one object in it.
-		FString ObjectName(PackageName);
-		ObjectName.Append(".TestObject");
-		UObject* const Object = NewObject<UMyAutoRTFMTestObject>(Package, *ObjectName, RF_Public | RF_Standalone);
+		UObject* const Object = NewObject<UMyAutoRTFMTestObject>(Package, ObjectName, RF_Public | RF_Standalone);
 
 		// Need to mark it is loaded.
 		Package->MarkAsFullyLoaded();
 
-		// Then wipe the standalone flag for reasons.
-		Object->ClearFlags(RF_Standalone);
-
 		// Save the package to the file-system.
-		REQUIRE(UPackage::SavePackage(Package, nullptr, *PackagePath, FSavePackageArgs()));
+		REQUIRE(UPackage::SavePackage(Package, Object, *PackagePath, FSavePackageArgs()));
 
 		// Make sure the package existed in our tables before.
 		REQUIRE(FindObject<UObject>(nullptr, *PackageName) != nullptr);
@@ -415,7 +404,7 @@ TEST_CASE("UPackage.AsyncLoading")
 		UPackage* const Package = LoadPackage(nullptr, *Name, LOAD_None);
 		REQUIRE(Package);
 			
-		TrashPackage(Package);
+		TrashObject(Package);
 
 		AutoRTFM::Testing::Commit([&]
 		{
@@ -458,7 +447,7 @@ TEST_CASE("UPackage.AsyncLoading")
 			UPackage* const Package = LoadPackage(nullptr, *Name, LOAD_None);
 			REQUIRE(Package);
 			
-			TrashPackage(Package);
+			TrashObject(Package);
 
 			UPackage* const ReloadedPackage = LoadPackage(nullptr, *Name, LOAD_None);
 			REQUIRE(ReloadedPackage);
@@ -480,7 +469,7 @@ TEST_CASE("UPackage.AsyncLoading")
 			UPackage* const FoundPackage = FindObject<UPackage>(nullptr, *Name);
 			REQUIRE(Package == FoundPackage);
 			
-			TrashPackage(Package);
+			TrashObject(Package);
 
 			UPackage* const ReloadedPackage = LoadPackage(nullptr, *Name, LOAD_None);
 			REQUIRE(ReloadedPackage);
@@ -542,7 +531,7 @@ TEST_CASE("UPackage.AsyncLoading")
 				UPackage* const Package = LoadPackage(nullptr, *Name, LOAD_None);
 				REQUIRE(Package);
 
-				TrashPackage(Package);
+				TrashObject(Package);
 
 				UPackage* const ReloadedPackage = LoadPackage(nullptr, *Name, LOAD_None);
 				REQUIRE(ReloadedPackage);
@@ -609,7 +598,7 @@ TEST_CASE("UPackage.AsyncLoading")
 
 		AutoRTFM::Testing::Commit([&]
 		{
-			TrashPackage(Package);
+			TrashObject(Package);
 
 			FName CustomName("Wowwee");
 
@@ -669,6 +658,71 @@ TEST_CASE("UPackage.AsyncLoading")
 		REQUIRE(bHit);
 	}
 
+	SECTION("Object is correct during trashing")
+	{
+		FString Name(FString::Printf(TEXT("/Game/%dAutoRTFMTestPackage%d"), FPlatformProcess::GetCurrentProcessId(), __LINE__));
+		FPackageScopedMaker _(Name);
+
+		UPackage* const Package = LoadPackage(nullptr, *Name, LOAD_None);
+		REQUIRE(Package);
+
+		UMyAutoRTFMTestObject* const Object = FindObject<UMyAutoRTFMTestObject>(Package, ObjectName);
+		REQUIRE(Object);
+
+		REQUIRE(Object->Value == 42);
+
+		Object->Value = 13;
+
+		AutoRTFM::Testing::Commit([&]
+		{
+			REQUIRE(Object->Value == 13);
+			TrashObject(Package);
+
+			UPackage* const ReloadedPackage = LoadPackage(nullptr, *Name, LOAD_None);
+			REQUIRE(ReloadedPackage);
+
+			UMyAutoRTFMTestObject* const ReloadedObject = FindObject<UMyAutoRTFMTestObject>(ReloadedPackage, ObjectName);
+			REQUIRE(ReloadedObject);
+
+			REQUIRE(ReloadedObject->Value == 42);
+		});
+	}
+
+	SECTION("Object has correct name")
+	{
+		FString Name(FString::Printf(TEXT("/Game/%dAutoRTFMTestPackage%d"), FPlatformProcess::GetCurrentProcessId(), __LINE__));
+		FPackageScopedMaker _(Name);
+
+		UPackage* const Package = LoadPackage(nullptr, *Name, LOAD_None);
+		REQUIRE(Package);
+
+		FString ObjectPath = Name;
+		ObjectPath.Append(TEXT("."));
+		ObjectPath.Append(ObjectName);
+
+		UMyAutoRTFMTestObject* const Object = FindObject<UMyAutoRTFMTestObject>(nullptr, ObjectPath);
+		REQUIRE(Object);
+	}
+
+	SECTION("Object has correct name in transaction")
+	{
+		FString Name(FString::Printf(TEXT("/Game/%dAutoRTFMTestPackage%d"), FPlatformProcess::GetCurrentProcessId(), __LINE__));
+		FPackageScopedMaker _(Name);
+
+		AutoRTFM::Testing::Commit([&]
+		{
+			UPackage* const Package = LoadPackage(nullptr, *Name, LOAD_None);
+			REQUIRE(Package);
+
+			FString ObjectPath = Name;
+			ObjectPath.Append(TEXT("."));
+			ObjectPath.Append(ObjectName);
+
+			UMyAutoRTFMTestObject* const Object = FindObject<UMyAutoRTFMTestObject>(nullptr, ObjectPath);
+			REQUIRE(Object);
+		});
+	}
+
 	SECTION("Create unrelated UObject's before package load")
 	{
 		FString Name(FString::Printf(TEXT("/Game/%dAutoRTFMTestPackage%d"), FPlatformProcess::GetCurrentProcessId(), __LINE__));
diff --git a/Engine/Source/Runtime/CoreUObject/Internal/UObject/UObjectGlobalsInternal.h b/Engine/Source/Runtime/CoreUObject/Internal/UObject/UObjectGlobalsInternal.h
--- a/Engine/Source/Runtime/CoreUObject/Internal/UObject/UObjectGlobalsInternal.h
+++ b/Engine/Source/Runtime/CoreUObject/Internal/UObject/UObjectGlobalsInternal.h
@@ -16,4 +16,11 @@ struct FCoreUObjectInternalDelegates
 	static COREUOBJECT_API FPackageRename& GetOnLeakedPackageRenameDelegate();
 };
 
-    
+/**
+ * Trash a `UObject`. Should be used with extreme caution. Used by systems that
+ * need to force reload objects from disk and thus need to not find the original
+ * object alive and well.
+ * 
+ * It is illegal to trash an object that is in the process of being loaded.
+ */
+void COREUOBJECT_API TrashObject(TNotNull<UObject* const> Object);
diff --git a/Engine/Source/Runtime/CoreUObject/Private/UObject/UObjectGlobals.cpp b/Engine/Source/Runtime/CoreUObject/Private/UObject/UObjectGlobals.cpp
--- a/Engine/Source/Runtime/CoreUObject/Private/UObject/UObjectGlobals.cpp
+++ b/Engine/Source/Runtime/CoreUObject/Private/UObject/UObjectGlobals.cpp
@@ -247,6 +247,29 @@ FCoreUObjectInternalDelegates::FPackageRename& FCoreUObjectInternalDelegates::Ge
 	return Delegate;
 }
 
+void TrashObject(TNotNull<UObject* const> Object)
+{
+	// We should *never* trash an object that is in the process of being loaded.
+	// TODO: This should probably degrade to a `check` longer term, but out of
+	// an abundance of caution we'll start it as an `ensure`.
+
+	const bool bIsBeingLoaded = Object->HasAnyFlags(EObjectFlags::RF_WillBeLoaded) ||
+		Object->HasAnyInternalFlags(EInternalObjectFlags_AsyncLoading);
+
+	if (ensure(!bIsBeingLoaded))
+	{
+		const bool bIsAPackage = Object->IsA<UPackage>();
+		UObject* const Outer = bIsAPackage ? nullptr : GetTransientPackage();
+		const FName BaseName = bIsAPackage ? NAME_TrashedPackage : NAME_None;
+
+		const FName NewName = MakeUniqueObjectName(Outer, Object->GetClass(), BaseName);
+		Object->Rename(*NewName.ToString(), Outer, REN_DontCreateRedirectors | REN_NonTransactional | REN_DoNotDirty);
+
+		// Mark it transient to be sure it is not seen as relevant by various systems.
+		Object->SetFlags(RF_Transient);
+	}
+}
+
 FSimpleMulticastDelegate FCoreUObjectDelegates::PostReachabilityAnalysis;
 
 FSimpleMulticastDelegate FCoreUObjectDelegates::PreGarbageCollectConditionalBeginDestroy;
diff --git a/Engine/Source/Runtime/Engine/Classes/Engine/LevelStreamingGCHelper.h b/Engine/Source/Runtime/Engine/Classes/Engine/LevelStreamingGCHelper.h
--- a/Engine/Source/Runtime/Engine/Classes/Engine/LevelStreamingGCHelper.h
+++ b/Engine/Source/Runtime/Engine/Classes/Engine/LevelStreamingGCHelper.h
@@ -43,6 +43,7 @@ struct FLevelStreamingGCHelper
 	/**
 	 * Trash a package, making sure it will be reloaded if needed again by level streaming
 	 */
+	UE_DEPRECATED(5.8, "TrashPackage is deprecated.")
 	static ENGINE_API void TrashPackage(UPackage* InPackage);
 
 	/**
diff --git a/Engine/Source/Runtime/Engine/Private/Level.cpp b/Engine/Source/Runtime/Engine/Private/Level.cpp
--- a/Engine/Source/Runtime/Engine/Private/Level.cpp
+++ b/Engine/Source/Runtime/Engine/Private/Level.cpp
@@ -85,6 +85,7 @@ Level.cpp: Level-related functions
 #include "WorldPartition/WorldPartitionRuntimeCell.h"
 #include "LevelInstance/LevelInstanceSubsystem.h"
 #include "Components/UnregisterComponentContext.h"
+#include "UObject/UObjectGlobalsInternal.h"
 
 #include UE_INLINE_GENERATED_CPP_BY_NAME(Level)
 
@@ -614,7 +615,7 @@ void ULevel::CleanupLevel(bool bCleanupResources, bool bUnloadFromEditor)
 			if (bTrashPackage && (InPackage != GetTransientPackage()))
 			{
 				// Rename package to make sure it won't be reused
-				FLevelStreamingGCHelper::TrashPackage(InPackage);
+				TrashObject(InPackage);
 			}
 		}
 	};
diff --git a/Engine/Source/Runtime/Engine/Private/LevelStreaming.cpp b/Engine/Source/Runtime/Engine/Private/LevelStreaming.cpp
--- a/Engine/Source/Runtime/Engine/Private/LevelStreaming.cpp
+++ b/Engine/Source/Runtime/Engine/Private/LevelStreaming.cpp
@@ -35,6 +35,7 @@
 #include "AssetRegistry/AssetRegistryModule.h"
 #include "WorldPartition/WorldPartition.h"
 #include "Misc/ScopeRWLock.h"
+#include "UObject/UObjectGlobalsInternal.h"
 
 #include UE_INLINE_GENERATED_CPP_BY_NAME(LevelStreaming)
 
@@ -1696,7 +1697,7 @@ bool ULevelStreaming::RequestLevel(UWorld* PersistentWorld, bool bAllowLevelLoad
 			if (!ULevelStreaming::ShouldReuseUnloadedButStillAroundLevels(World->PersistentLevel))
 			{
 				UE_LOG(LogLevelStreaming, Verbose, TEXT("Trying to reuse level %s but ShouldReuseUnloadedButStillAroundLevels returns false. Trash the current package and reload."), *LevelPackage->GetName());
-				FLevelStreamingGCHelper::TrashPackage(LevelPackage);
+				TrashObject(LevelPackage);
 				return false;
 			}
 
diff --git a/Engine/Source/Runtime/Engine/Private/LevelStreamingGCHelper.cpp b/Engine/Source/Runtime/Engine/Private/LevelStreamingGCHelper.cpp
--- a/Engine/Source/Runtime/Engine/Private/LevelStreamingGCHelper.cpp
+++ b/Engine/Source/Runtime/Engine/Private/LevelStreamingGCHelper.cpp
@@ -189,7 +189,7 @@ void FLevelStreamingGCHelper::PrepareStreamedOutLevelForGC(ULevel* InLevel)
 			for (UPackage* Package : Packages)
 			{
 				FCoreUObjectInternalDelegates::GetOnLeakedPackageRenameDelegate().Broadcast(Package);
-				TrashPackage(Package);
+				TrashObject(Package);
 			}
 
 #if !WITH_EDITOR
@@ -260,9 +260,5 @@ int32 FLevelStreamingGCHelper::GetNumLevelsPendingPurge()
 
 void FLevelStreamingGCHelper::TrashPackage(UPackage* InPackage)
 {
-	const FName NewName = MakeUniqueObjectName(nullptr, UPackage::StaticClass(), NAME_TrashedPackage);
-	InPackage->Rename(*NewName.ToString(), nullptr, REN_DontCreateRedirectors | REN_NonTransactional | REN_DoNotDirty);
-
-	// Mark it transient to be sure it is not seen as relevant by various systems
-	InPackage->SetFlags(RF_Transient);
+	TrashObject(InPackage);
 }
diff --git a/Engine/Source/Runtime/Engine/Private/WorldPartition/RuntimeHashSet/WorldPartitionRuntimeHashSet.cpp b/Engine/Source/Runtime/Engine/Private/WorldPartition/RuntimeHashSet/WorldPartitionRuntimeHashSet.cpp
--- a/Engine/Source/Runtime/Engine/Private/WorldPartition/RuntimeHashSet/WorldPartitionRuntimeHashSet.cpp
+++ b/Engine/Source/Runtime/Engine/Private/WorldPartition/RuntimeHashSet/WorldPartitionRuntimeHashSet.cpp
@@ -16,6 +16,7 @@
 #include "Algo/Transform.h"
 #include "Misc/HashBuilder.h"
 #include "Misc/ArchiveMD5.h"
+#include "UObject/UObjectGlobalsInternal.h"
 
 #include "WorldPartition/RuntimeHashSet/StaticSpatialIndex.h"
 
@@ -905,13 +906,7 @@ bool UWorldPartitionRuntimeHashSet::UnregisterWorldAssetStreaming(const FGuid& I
 		// Apply the same logic on the LevelStreaming object of each injected cell as it is named using the injected cell name and is outered to the owning world.
 		auto TrashExternalStreamingData = [](URuntimeHashSetExternalStreamingObject* InStreamingObject)
 		{
-			auto TrashObject = [](UObject* InObject)
-			{
-				FName NewUniqueTrashName = MakeUniqueObjectName(InObject->GetOuter(), InObject->GetClass(), NAME_TrashedPackage);
-				InObject->Rename(*NewUniqueTrashName.ToString(), nullptr, REN_DontCreateRedirectors | REN_NonTransactional | REN_DoNotDirty);
-			};
-
-			InStreamingObject->ForEachStreamingCells([TrashObject](UWorldPartitionRuntimeCell& Cell)
+			InStreamingObject->ForEachStreamingCells([](UWorldPartitionRuntimeCell& Cell)
 			{
 				UWorldPartitionRuntimeLevelStreamingCell* InjectedCell = Cast<UWorldPartitionRuntimeLevelStreamingCell>(&Cell);
 				if (UWorldPartitionLevelStreamingDynamic* LevelStreaming = InjectedCell ? InjectedCell->GetLevelStreaming() : nullptr)
diff --git a/Engine/Source/Runtime/Engine/Private/WorldPartition/WorldPartitionLevelHelper.cpp b/Engine/Source/Runtime/Engine/Private/WorldPartition/WorldPartitionLevelHelper.cpp
--- a/Engine/Source/Runtime/Engine/Private/WorldPartition/WorldPartitionLevelHelper.cpp
+++ b/Engine/Source/Runtime/Engine/Private/WorldPartition/WorldPartitionLevelHelper.cpp
@@ -23,6 +23,7 @@
 #include "GameFramework/WorldSettings.h"
 #include "LevelUtils.h"
 #include "ActorFolder.h"
+#include "UObject/UObjectGlobalsInternal.h"
 
 FUObjectAnnotationSparse<FWorldPartitionLevelHelper::FActorPropertyOverridesAnnotation, true> FWorldPartitionLevelHelper::ActorPropertyOverridesAnnotation;
 #endif
@@ -362,7 +363,7 @@ void FWorldPartitionLevelHelper::MoveExternalActorsToLevel(const TArray<FWorldPa
 	{
 		if (ModifiedPackage != LevelPackage)
 		{
-			FLevelStreamingGCHelper::TrashPackage(ModifiedPackage);
+			TrashObject(ModifiedPackage);
 		}
 	}
 
diff --git a/Engine/Source/Runtime/Engine/Private/WorldPartition/WorldPartitionPackageHelper.cpp b/Engine/Source/Runtime/Engine/Private/WorldPartition/WorldPartitionPackageHelper.cpp
--- a/Engine/Source/Runtime/Engine/Private/WorldPartition/WorldPartitionPackageHelper.cpp
+++ b/Engine/Source/Runtime/Engine/Private/WorldPartition/WorldPartitionPackageHelper.cpp
@@ -7,16 +7,11 @@
 #include "Engine/Level.h"
 #include "Engine/LevelStreamingGCHelper.h"
 #include "Engine/World.h"
+#include "UObject/UObjectGlobalsInternal.h"
 
 void FWorldPartitionPackageHelper::UnloadPackage(UPackage* InPackage)
 {
-	auto TrashPackage = [](UPackage* InPackage)
-	{
-		// Rename so it isn't found again
-		FLevelStreamingGCHelper::TrashPackage(InPackage);
-	};
-
-	TrashPackage(InPackage);
+	TrashObject(InPackage);
 
 	// World specific
 	if (UWorld* PackageWorld = UWorld::FindWorldInPackage(InPackage))
@@ -30,11 +25,11 @@ void FWorldPartitionPackageHelper::UnloadPackage(UPackage* InPackage)
 
 			if (PackageWorld->PersistentLevel->IsUsingExternalObjects())
 			{
-				ForEachObjectWithOuter(PackageWorld->PersistentLevel, [&TrashPackage](UObject* InObject)
+				ForEachObjectWithOuter(PackageWorld->PersistentLevel, [](UObject* InObject)
 				{
 					if (UPackage* ExternalPackage = InObject->GetExternalPackage())
 					{
-						TrashPackage(ExternalPackage);
+						TrashObject(ExternalPackage);
 					}
 				}, /*bIncludeNestedObjects*/ true);
 			}
diff --git a/Engine/Source/Runtime/Engine/Private/WorldPartition/WorldPartitionRuntimeSpatialHash.cpp b/Engine/Source/Runtime/Engine/Private/WorldPartition/WorldPartitionRuntimeSpatialHash.cpp
--- a/Engine/Source/Runtime/Engine/Private/WorldPartition/WorldPartitionRuntimeSpatialHash.cpp
+++ b/Engine/Source/Runtime/Engine/Private/WorldPartition/WorldPartitionRuntimeSpatialHash.cpp
@@ -44,6 +44,7 @@
 #include "UObject/FortniteMainBranchObjectVersion.h"
 #include "Misc/HashBuilder.h"
 #include "Algo/AllOf.h"
+#include "UObject/UObjectGlobalsInternal.h"
 
 #include UE_INLINE_GENERATED_CPP_BY_NAME(WorldPartitionRuntimeSpatialHash)
 
@@ -2212,13 +2213,7 @@ bool UWorldPartitionRuntimeSpatialHash::UnregisterWorldAssetStreaming(const FGui
 		// Apply the same logic on the LevelStreaming object of each injected cell as it is named using the injected cell name and is outered to the owning world.
 		auto TrashExternalStreamingData = [](URuntimeSpatialHashExternalStreamingObject* InStreamingObject)
 		{
-			auto TrashObject = [](UObject* InObject)
-			{
-				FName NewUniqueTrashName = MakeUniqueObjectName(InObject->GetOuter(), InObject->GetClass(), NAME_TrashedPackage);
-				InObject->Rename(*NewUniqueTrashName.ToString(), nullptr, REN_DontCreateRedirectors | REN_NonTransactional | REN_DoNotDirty);
-			};
-
-			InStreamingObject->ForEachStreamingCells([TrashObject](UWorldPartitionRuntimeCell& Cell)
+			InStreamingObject->ForEachStreamingCells([](UWorldPartitionRuntimeCell& Cell)
 			{
 				UWorldPartitionRuntimeLevelStreamingCell* InjectedCell = Cast<UWorldPartitionRuntimeLevelStreamingCell>(&Cell);
 				if (UWorldPartitionLevelStreamingDynamic* LevelStreaming = InjectedCell ? InjectedCell->GetLevelStreaming() : nullptr)
diff --git a/Engine/Source/Runtime/Engine/Private/WorldPartition/WorldPartitionStreamingGeneration.cpp b/Engine/Source/Runtime/Engine/Private/WorldPartition/WorldPartitionStreamingGeneration.cpp
--- a/Engine/Source/Runtime/Engine/Private/WorldPartition/WorldPartitionStreamingGeneration.cpp
+++ b/Engine/Source/Runtime/Engine/Private/WorldPartition/WorldPartitionStreamingGeneration.cpp
@@ -46,6 +46,7 @@
 #include "WorldPartition/WorldPartitionRuntimeContainerResolving.h"
 #include "HAL/FileManager.h"
 #include "Misc/EditorPathHelper.h"
+#include "UObject/UObjectGlobalsInternal.h"
 
 #define LOCTEXT_NAMESPACE "WorldPartition"
 
@@ -954,8 +955,7 @@ class FWorldPartitionStreamingGenerator
 
 					// Destroy placeholder actor since we only need its actor descriptor.
 					WorldPartitionContext->GetWorld()->DestroyActor(CustomHLODPlaceholderActor, false, false);
-					FName UniqueTrashName = MakeUniqueObjectName(CustomHLODPlaceholderActor->GetOuter(), CustomHLODPlaceholderActor->GetClass(), NAME_TrashedPackage);
-					CustomHLODPlaceholderActor->Rename(*UniqueTrashName.ToString(), nullptr, REN_DontCreateRedirectors | REN_NonTransactional | REN_DoNotDirty);
+					TrashObject(CustomHLODPlaceholderActor);
 				}
 
 				// If Custom HLOD actor is used with some HLOD layer different than CustomHLODActor layer don't register its actor descriptor
