[Backout] - CL46934445
#rnx
Original CL Desc
-----------------------------------------------------------------
Make a internal API for trashing objects and use it everywhere we previously had lots of copy-pasta code.

diff --git a/Engine/Source/Programs/AutoRTFMTests/AutoRTFMTests.Build.cs b/Engine/Source/Programs/AutoRTFMTests/AutoRTFMTests.Build.cs
--- a/Engine/Source/Programs/AutoRTFMTests/AutoRTFMTests.Build.cs
+++ b/Engine/Source/Programs/AutoRTFMTests/AutoRTFMTests.Build.cs
@@ -1,6 +1,5 @@
 // Copyright Epic Games, Inc. All Rights Reserved.
 
-using System.IO;
 using UnrealBuildTool;
 
 public class AutoRTFMTests : ModuleRules
@@ -30,9 +29,8 @@ public class AutoRTFMTests : ModuleRules
 		PrivateIncludePaths.AddRange(
 			new string[]
 			{
-				Path.Combine(GetModuleDirectory("AutoRTFM"), "Private"),
-				Path.Combine(GetModuleDirectory("Core"), "Private"),
-				Path.Combine(GetModuleDirectory("CoreUObject"), "Internal"),
+				"Runtime/AutoRTFM/Private",
+				"Runtime/Core/Private",
 			});
 
 		PCHUsage = PCHUsageMode.NoPCHs;
diff --git a/Engine/Source/Programs/AutoRTFMTests/Classes/MyAutoRTFMTestObject.h b/Engine/Source/Programs/AutoRTFMTests/Classes/MyAutoRTFMTestObject.h
--- a/Engine/Source/Programs/AutoRTFMTests/Classes/MyAutoRTFMTestObject.h
+++ b/Engine/Source/Programs/AutoRTFMTests/Classes/MyAutoRTFMTestObject.h
@@ -13,7 +13,7 @@ public:
 
 	UMyAutoRTFMTestObject(const FObjectInitializer& ObjectInitializer = FObjectInitializer::Get());
 
-	int Value = 42;
+	int Value;
 
 	static FConstructorCallback* ConstructorCallback;
 
diff --git a/Engine/Source/Programs/AutoRTFMTests/Private/UPackageTests.cpp b/Engine/Source/Programs/AutoRTFMTests/Private/UPackageTests.cpp
--- a/Engine/Source/Programs/AutoRTFMTests/Private/UPackageTests.cpp
+++ b/Engine/Source/Programs/AutoRTFMTests/Private/UPackageTests.cpp
@@ -8,10 +8,17 @@
 #include "UObject/LinkerInstancingContext.h"
 #include "UObject/Package.h"
 #include "UObject/SavePackage.h"
-#include "UObject/UObjectGlobalsInternal.h"
 
 #include "Catch2Includes.h"
 
+static void TrashPackage(UPackage* const Package)
+{
+	// TODO: If we could move the trashing logic into `UPackage` we could just call that here?
+	const FName NewName = MakeUniqueObjectName(nullptr, UPackage::StaticClass(), NAME_TrashedPackage);
+	Package->Rename(*NewName.ToString(), nullptr, REN_DontCreateRedirectors | REN_NonTransactional | REN_DoNotDirty);
+	Package->SetFlags(RF_Transient);
+}
+
 TEST_CASE("UPackage.SetPackageFlagsTo")
 {
 	SECTION("Commit")
@@ -105,8 +112,6 @@ TEST_CASE("UPackage.ClearPackageFlags")
 	}
 }
 
-static const TCHAR* ObjectName = TEXT("TestObject");
-
 // Based on `Engine\Private\Tests\Loading\AsyncLoadingTests_Shared.h`, we use similar logic
 // here to make a package that the loader will see and be able to actually load!
 struct FPackageScopedMaker final
@@ -128,13 +133,19 @@ struct FPackageScopedMaker final
 		// Create a package.
 		UPackage* const Package = CreatePackage(*PackageName);
 		
-		UObject* const Object = NewObject<UMyAutoRTFMTestObject>(Package, ObjectName, RF_Public | RF_Standalone);
+		// With at least one object in it.
+		FString ObjectName(PackageName);
+		ObjectName.Append(".TestObject");
+		UObject* const Object = NewObject<UMyAutoRTFMTestObject>(Package, *ObjectName, RF_Public | RF_Standalone);
 
 		// Need to mark it is loaded.
 		Package->MarkAsFullyLoaded();
 
+		// Then wipe the standalone flag for reasons.
+		Object->ClearFlags(RF_Standalone);
+
 		// Save the package to the file-system.
-		REQUIRE(UPackage::SavePackage(Package, Object, *PackagePath, FSavePackageArgs()));
+		REQUIRE(UPackage::SavePackage(Package, nullptr, *PackagePath, FSavePackageArgs()));
 
 		// Make sure the package existed in our tables before.
 		REQUIRE(FindObject<UObject>(nullptr, *PackageName) != nullptr);
@@ -404,7 +415,7 @@ TEST_CASE("UPackage.AsyncLoading")
 		UPackage* const Package = LoadPackage(nullptr, *Name, LOAD_None);
 		REQUIRE(Package);
 			
-		TrashObject(Package);
+		TrashPackage(Package);
 
 		AutoRTFM::Testing::Commit([&]
 		{
@@ -447,7 +458,7 @@ TEST_CASE("UPackage.AsyncLoading")
 			UPackage* const Package = LoadPackage(nullptr, *Name, LOAD_None);
 			REQUIRE(Package);
 			
-			TrashObject(Package);
+			TrashPackage(Package);
 
 			UPackage* const ReloadedPackage = LoadPackage(nullptr, *Name, LOAD_None);
 			REQUIRE(ReloadedPackage);
@@ -469,7 +480,7 @@ TEST_CASE("UPackage.AsyncLoading")
 			UPackage* const FoundPackage = FindObject<UPackage>(nullptr, *Name);
 			REQUIRE(Package == FoundPackage);
 			
-			TrashObject(Package);
+			TrashPackage(Package);
 
 			UPackage* const ReloadedPackage = LoadPackage(nullptr, *Name, LOAD_None);
 			REQUIRE(ReloadedPackage);
@@ -531,7 +542,7 @@ TEST_CASE("UPackage.AsyncLoading")
 				UPackage* const Package = LoadPackage(nullptr, *Name, LOAD_None);
 				REQUIRE(Package);
 
-				TrashObject(Package);
+				TrashPackage(Package);
 
 				UPackage* const ReloadedPackage = LoadPackage(nullptr, *Name, LOAD_None);
 				REQUIRE(ReloadedPackage);
@@ -598,7 +609,7 @@ TEST_CASE("UPackage.AsyncLoading")
 
 		AutoRTFM::Testing::Commit([&]
 		{
-			TrashObject(Package);
+			TrashPackage(Package);
 
 			FName CustomName("Wowwee");
 
@@ -658,71 +669,6 @@ TEST_CASE("UPackage.AsyncLoading")
 		REQUIRE(bHit);
 	}
 
-	SECTION("Object is correct during trashing")
-	{
-		FString Name(FString::Printf(TEXT("/Game/%dAutoRTFMTestPackage%d"), FPlatformProcess::GetCurrentProcessId(), __LINE__));
-		FPackageScopedMaker _(Name);
-
-		UPackage* const Package = LoadPackage(nullptr, *Name, LOAD_None);
-		REQUIRE(Package);
-
-		UMyAutoRTFMTestObject* const Object = FindObject<UMyAutoRTFMTestObject>(Package, ObjectName);
-		REQUIRE(Object);
-
-		REQUIRE(Object->Value == 42);
-
-		Object->Value = 13;
-
-		AutoRTFM::Testing::Commit([&]
-		{
-			REQUIRE(Object->Value == 13);
-			TrashObject(Package);
-
-			UPackage* const ReloadedPackage = LoadPackage(nullptr, *Name, LOAD_None);
-			REQUIRE(ReloadedPackage);
-
-			UMyAutoRTFMTestObject* const ReloadedObject = FindObject<UMyAutoRTFMTestObject>(ReloadedPackage, ObjectName);
-			REQUIRE(ReloadedObject);
-
-			REQUIRE(ReloadedObject->Value == 42);
-		});
-	}
-
-	SECTION("Object has correct name")
-	{
-		FString Name(FString::Printf(TEXT("/Game/%dAutoRTFMTestPackage%d"), FPlatformProcess::GetCurrentProcessId(), __LINE__));
-		FPackageScopedMaker _(Name);
-
-		UPackage* const Package = LoadPackage(nullptr, *Name, LOAD_None);
-		REQUIRE(Package);
-
-		FString ObjectPath = Name;
-		ObjectPath.Append(TEXT("."));
-		ObjectPath.Append(ObjectName);
-
-		UMyAutoRTFMTestObject* const Object = FindObject<UMyAutoRTFMTestObject>(nullptr, ObjectPath);
-		REQUIRE(Object);
-	}
-
-	SECTION("Object has correct name in transaction")
-	{
-		FString Name(FString::Printf(TEXT("/Game/%dAutoRTFMTestPackage%d"), FPlatformProcess::GetCurrentProcessId(), __LINE__));
-		FPackageScopedMaker _(Name);
-
-		AutoRTFM::Testing::Commit([&]
-		{
-			UPackage* const Package = LoadPackage(nullptr, *Name, LOAD_None);
-			REQUIRE(Package);
-
-			FString ObjectPath = Name;
-			ObjectPath.Append(TEXT("."));
-			ObjectPath.Append(ObjectName);
-
-			UMyAutoRTFMTestObject* const Object = FindObject<UMyAutoRTFMTestObject>(nullptr, ObjectPath);
-			REQUIRE(Object);
-		});
-	}
-
 	SECTION("Create unrelated UObject's before package load")
 	{
 		FString Name(FString::Printf(TEXT("/Game/%dAutoRTFMTestPackage%d"), FPlatformProcess::GetCurrentProcessId(), __LINE__));
diff --git a/Engine/Source/Runtime/CoreUObject/Internal/UObject/UObjectGlobalsInternal.h b/Engine/Source/Runtime/CoreUObject/Internal/UObject/UObjectGlobalsInternal.h
--- a/Engine/Source/Runtime/CoreUObject/Internal/UObject/UObjectGlobalsInternal.h
+++ b/Engine/Source/Runtime/CoreUObject/Internal/UObject/UObjectGlobalsInternal.h
@@ -16,11 +16,4 @@ struct FCoreUObjectInternalDelegates
 	static COREUOBJECT_API FPackageRename& GetOnLeakedPackageRenameDelegate();
 };
 
-/**
- * Trash a `UObject`. Should be used with extreme caution. Used by systems that
- * need to force reload objects from disk and thus need to not find the original
- * object alive and well.
- * 
- * It is illegal to trash an object that is in the process of being loaded.
- */
-void COREUOBJECT_API TrashObject(TNotNull<UObject* const> Object);
+    
diff --git a/Engine/Source/Runtime/CoreUObject/Private/UObject/UObjectGlobals.cpp b/Engine/Source/Runtime/CoreUObject/Private/UObject/UObjectGlobals.cpp
--- a/Engine/Source/Runtime/CoreUObject/Private/UObject/UObjectGlobals.cpp
+++ b/Engine/Source/Runtime/CoreUObject/Private/UObject/UObjectGlobals.cpp
@@ -247,29 +247,6 @@ FCoreUObjectInternalDelegates::FPackageRename& FCoreUObjectInternalDelegates::Ge
 	return Delegate;
 }
 
-void TrashObject(TNotNull<UObject* const> Object)
-{
-	// We should *never* trash an object that is in the process of being loaded.
-	// TODO: This should probably degrade to a `check` longer term, but out of
-	// an abundance of caution we'll start it as an `ensure`.
-
-	const bool bIsBeingLoaded = Object->HasAnyFlags(EObjectFlags::RF_WillBeLoaded) ||
-		Object->HasAnyInternalFlags(EInternalObjectFlags_AsyncLoading);
-
-	if (ensure(!bIsBeingLoaded))
-	{
-		const bool bIsAPackage = Object->IsA<UPackage>();
-		UObject* const Outer = bIsAPackage ? nullptr : GetTransientPackage();
-		const FName BaseName = bIsAPackage ? NAME_TrashedPackage : NAME_None;
-
-		const FName NewName = MakeUniqueObjectName(Outer, Object->GetClass(), BaseName);
-		Object->Rename(*NewName.ToString(), Outer, REN_DontCreateRedirectors | REN_NonTransactional | REN_DoNotDirty);
-
-		// Mark it transient to be sure it is not seen as relevant by various systems.
-		Object->SetFlags(RF_Transient);
-	}
-}
-
 FSimpleMulticastDelegate FCoreUObjectDelegates::PostReachabilityAnalysis;
 
 FSimpleMulticastDelegate FCoreUObjectDelegates::PreGarbageCollectConditionalBeginDestroy;
diff --git a/Engine/Source/Runtime/Engine/Classes/Engine/LevelStreamingGCHelper.h b/Engine/Source/Runtime/Engine/Classes/Engine/LevelStreamingGCHelper.h
--- a/Engine/Source/Runtime/Engine/Classes/Engine/LevelStreamingGCHelper.h
+++ b/Engine/Source/Runtime/Engine/Classes/Engine/LevelStreamingGCHelper.h
@@ -43,7 +43,6 @@ struct FLevelStreamingGCHelper
 	/**
 	 * Trash a package, making sure it will be reloaded if needed again by level streaming
 	 */
-	UE_DEPRECATED(5.8, "TrashPackage is deprecated.")
 	static ENGINE_API void TrashPackage(UPackage* InPackage);
 
 	/**
diff --git a/Engine/Source/Runtime/Engine/Private/Level.cpp b/Engine/Source/Runtime/Engine/Private/Level.cpp
--- a/Engine/Source/Runtime/Engine/Private/Level.cpp
+++ b/Engine/Source/Runtime/Engine/Private/Level.cpp
@@ -85,7 +85,6 @@ Level.cpp: Level-related functions
 #include "WorldPartition/WorldPartitionRuntimeCell.h"
 #include "LevelInstance/LevelInstanceSubsystem.h"
 #include "Components/UnregisterComponentContext.h"
-#include "UObject/UObjectGlobalsInternal.h"
 
 #include UE_INLINE_GENERATED_CPP_BY_NAME(Level)
 
@@ -615,7 +614,7 @@ void ULevel::CleanupLevel(bool bCleanupResources, bool bUnloadFromEditor)
 			if (bTrashPackage && (InPackage != GetTransientPackage()))
 			{
 				// Rename package to make sure it won't be reused
-				TrashObject(InPackage);
+				FLevelStreamingGCHelper::TrashPackage(InPackage);
 			}
 		}
 	};
diff --git a/Engine/Source/Runtime/Engine/Private/LevelStreaming.cpp b/Engine/Source/Runtime/Engine/Private/LevelStreaming.cpp
--- a/Engine/Source/Runtime/Engine/Private/LevelStreaming.cpp
+++ b/Engine/Source/Runtime/Engine/Private/LevelStreaming.cpp
@@ -35,7 +35,6 @@
 #include "AssetRegistry/AssetRegistryModule.h"
 #include "WorldPartition/WorldPartition.h"
 #include "Misc/ScopeRWLock.h"
-#include "UObject/UObjectGlobalsInternal.h"
 
 #include UE_INLINE_GENERATED_CPP_BY_NAME(LevelStreaming)
 
@@ -1697,7 +1696,7 @@ bool ULevelStreaming::RequestLevel(UWorld* PersistentWorld, bool bAllowLevelLoad
 			if (!ULevelStreaming::ShouldReuseUnloadedButStillAroundLevels(World->PersistentLevel))
 			{
 				UE_LOG(LogLevelStreaming, Verbose, TEXT("Trying to reuse level %s but ShouldReuseUnloadedButStillAroundLevels returns false. Trash the current package and reload."), *LevelPackage->GetName());
-				TrashObject(LevelPackage);
+				FLevelStreamingGCHelper::TrashPackage(LevelPackage);
 				return false;
 			}
 
diff --git a/Engine/Source/Runtime/Engine/Private/LevelStreamingGCHelper.cpp b/Engine/Source/Runtime/Engine/Private/LevelStreamingGCHelper.cpp
--- a/Engine/Source/Runtime/Engine/Private/LevelStreamingGCHelper.cpp
+++ b/Engine/Source/Runtime/Engine/Private/LevelStreamingGCHelper.cpp
@@ -189,7 +189,7 @@ void FLevelStreamingGCHelper::PrepareStreamedOutLevelForGC(ULevel* InLevel)
 			for (UPackage* Package : Packages)
 			{
 				FCoreUObjectInternalDelegates::GetOnLeakedPackageRenameDelegate().Broadcast(Package);
-				TrashObject(Package);
+				TrashPackage(Package);
 			}
 
 #if !WITH_EDITOR
@@ -260,5 +260,9 @@ int32 FLevelStreamingGCHelper::GetNumLevelsPendingPurge()
 
 void FLevelStreamingGCHelper::TrashPackage(UPackage* InPackage)
 {
-	TrashObject(InPackage);
+	const FName NewName = MakeUniqueObjectName(nullptr, UPackage::StaticClass(), NAME_TrashedPackage);
+	InPackage->Rename(*NewName.ToString(), nullptr, REN_DontCreateRedirectors | REN_NonTransactional | REN_DoNotDirty);
+
+	// Mark it transient to be sure it is not seen as relevant by various systems
+	InPackage->SetFlags(RF_Transient);
 }
diff --git a/Engine/Source/Runtime/Engine/Private/WorldPartition/RuntimeHashSet/WorldPartitionRuntimeHashSet.cpp b/Engine/Source/Runtime/Engine/Private/WorldPartition/RuntimeHashSet/WorldPartitionRuntimeHashSet.cpp
--- a/Engine/Source/Runtime/Engine/Private/WorldPartition/RuntimeHashSet/WorldPartitionRuntimeHashSet.cpp
+++ b/Engine/Source/Runtime/Engine/Private/WorldPartition/RuntimeHashSet/WorldPartitionRuntimeHashSet.cpp
@@ -16,7 +16,6 @@
 #include "Algo/Transform.h"
 #include "Misc/HashBuilder.h"
 #include "Misc/ArchiveMD5.h"
-#include "UObject/UObjectGlobalsInternal.h"
 
 #include "WorldPartition/RuntimeHashSet/StaticSpatialIndex.h"
 
@@ -906,7 +905,13 @@ bool UWorldPartitionRuntimeHashSet::UnregisterWorldAssetStreaming(const FGuid& I
 		// Apply the same logic on the LevelStreaming object of each injected cell as it is named using the injected cell name and is outered to the owning world.
 		auto TrashExternalStreamingData = [](URuntimeHashSetExternalStreamingObject* InStreamingObject)
 		{
-			InStreamingObject->ForEachStreamingCells([](UWorldPartitionRuntimeCell& Cell)
+			auto TrashObject = [](UObject* InObject)
+			{
+				FName NewUniqueTrashName = MakeUniqueObjectName(InObject->GetOuter(), InObject->GetClass(), NAME_TrashedPackage);
+				InObject->Rename(*NewUniqueTrashName.ToString(), nullptr, REN_DontCreateRedirectors | REN_NonTransactional | REN_DoNotDirty);
+			};
+
+			InStreamingObject->ForEachStreamingCells([TrashObject](UWorldPartitionRuntimeCell& Cell)
 			{
 				UWorldPartitionRuntimeLevelStreamingCell* InjectedCell = Cast<UWorldPartitionRuntimeLevelStreamingCell>(&Cell);
 				if (UWorldPartitionLevelStreamingDynamic* LevelStreaming = InjectedCell ? InjectedCell->GetLevelStreaming() : nullptr)
diff --git a/Engine/Source/Runtime/Engine/Private/WorldPartition/WorldPartitionLevelHelper.cpp b/Engine/Source/Runtime/Engine/Private/WorldPartition/WorldPartitionLevelHelper.cpp
--- a/Engine/Source/Runtime/Engine/Private/WorldPartition/WorldPartitionLevelHelper.cpp
+++ b/Engine/Source/Runtime/Engine/Private/WorldPartition/WorldPartitionLevelHelper.cpp
@@ -23,7 +23,6 @@
 #include "GameFramework/WorldSettings.h"
 #include "LevelUtils.h"
 #include "ActorFolder.h"
-#include "UObject/UObjectGlobalsInternal.h"
 
 FUObjectAnnotationSparse<FWorldPartitionLevelHelper::FActorPropertyOverridesAnnotation, true> FWorldPartitionLevelHelper::ActorPropertyOverridesAnnotation;
 #endif
@@ -363,7 +362,7 @@ void FWorldPartitionLevelHelper::MoveExternalActorsToLevel(const TArray<FWorldPa
 	{
 		if (ModifiedPackage != LevelPackage)
 		{
-			TrashObject(ModifiedPackage);
+			FLevelStreamingGCHelper::TrashPackage(ModifiedPackage);
 		}
 	}
 
diff --git a/Engine/Source/Runtime/Engine/Private/WorldPartition/WorldPartitionPackageHelper.cpp b/Engine/Source/Runtime/Engine/Private/WorldPartition/WorldPartitionPackageHelper.cpp
--- a/Engine/Source/Runtime/Engine/Private/WorldPartition/WorldPartitionPackageHelper.cpp
+++ b/Engine/Source/Runtime/Engine/Private/WorldPartition/WorldPartitionPackageHelper.cpp
@@ -7,11 +7,16 @@
 #include "Engine/Level.h"
 #include "Engine/LevelStreamingGCHelper.h"
 #include "Engine/World.h"
-#include "UObject/UObjectGlobalsInternal.h"
 
 void FWorldPartitionPackageHelper::UnloadPackage(UPackage* InPackage)
 {
-	TrashObject(InPackage);
+	auto TrashPackage = [](UPackage* InPackage)
+	{
+		// Rename so it isn't found again
+		FLevelStreamingGCHelper::TrashPackage(InPackage);
+	};
+
+	TrashPackage(InPackage);
 
 	// World specific
 	if (UWorld* PackageWorld = UWorld::FindWorldInPackage(InPackage))
@@ -25,11 +30,11 @@ void FWorldPartitionPackageHelper::UnloadPackage(UPackage* InPackage)
 
 			if (PackageWorld->PersistentLevel->IsUsingExternalObjects())
 			{
-				ForEachObjectWithOuter(PackageWorld->PersistentLevel, [](UObject* InObject)
+				ForEachObjectWithOuter(PackageWorld->PersistentLevel, [&TrashPackage](UObject* InObject)
 				{
 					if (UPackage* ExternalPackage = InObject->GetExternalPackage())
 					{
-						TrashObject(ExternalPackage);
+						TrashPackage(ExternalPackage);
 					}
 				}, /*bIncludeNestedObjects*/ true);
 			}
diff --git a/Engine/Source/Runtime/Engine/Private/WorldPartition/WorldPartitionRuntimeSpatialHash.cpp b/Engine/Source/Runtime/Engine/Private/WorldPartition/WorldPartitionRuntimeSpatialHash.cpp
--- a/Engine/Source/Runtime/Engine/Private/WorldPartition/WorldPartitionRuntimeSpatialHash.cpp
+++ b/Engine/Source/Runtime/Engine/Private/WorldPartition/WorldPartitionRuntimeSpatialHash.cpp
@@ -44,7 +44,6 @@
 #include "UObject/FortniteMainBranchObjectVersion.h"
 #include "Misc/HashBuilder.h"
 #include "Algo/AllOf.h"
-#include "UObject/UObjectGlobalsInternal.h"
 
 #include UE_INLINE_GENERATED_CPP_BY_NAME(WorldPartitionRuntimeSpatialHash)
 
@@ -2222,7 +2221,13 @@ bool UWorldPartitionRuntimeSpatialHash::UnregisterWorldAssetStreaming(const FGui
 		// Apply the same logic on the LevelStreaming object of each injected cell as it is named using the injected cell name and is outered to the owning world.
 		auto TrashExternalStreamingData = [](URuntimeSpatialHashExternalStreamingObject* InStreamingObject)
 		{
-			InStreamingObject->ForEachStreamingCells([](UWorldPartitionRuntimeCell& Cell)
+			auto TrashObject = [](UObject* InObject)
+			{
+				FName NewUniqueTrashName = MakeUniqueObjectName(InObject->GetOuter(), InObject->GetClass(), NAME_TrashedPackage);
+				InObject->Rename(*NewUniqueTrashName.ToString(), nullptr, REN_DontCreateRedirectors | REN_NonTransactional | REN_DoNotDirty);
+			};
+
+			InStreamingObject->ForEachStreamingCells([TrashObject](UWorldPartitionRuntimeCell& Cell)
 			{
 				UWorldPartitionRuntimeLevelStreamingCell* InjectedCell = Cast<UWorldPartitionRuntimeLevelStreamingCell>(&Cell);
 				if (UWorldPartitionLevelStreamingDynamic* LevelStreaming = InjectedCell ? InjectedCell->GetLevelStreaming() : nullptr)
diff --git a/Engine/Source/Runtime/Engine/Private/WorldPartition/WorldPartitionStreamingGeneration.cpp b/Engine/Source/Runtime/Engine/Private/WorldPartition/WorldPartitionStreamingGeneration.cpp
--- a/Engine/Source/Runtime/Engine/Private/WorldPartition/WorldPartitionStreamingGeneration.cpp
+++ b/Engine/Source/Runtime/Engine/Private/WorldPartition/WorldPartitionStreamingGeneration.cpp
@@ -46,7 +46,6 @@
 #include "WorldPartition/WorldPartitionRuntimeContainerResolving.h"
 #include "HAL/FileManager.h"
 #include "Misc/EditorPathHelper.h"
-#include "UObject/UObjectGlobalsInternal.h"
 
 #define LOCTEXT_NAMESPACE "WorldPartition"
 
@@ -955,7 +954,8 @@ class FWorldPartitionStreamingGenerator
 
 					// Destroy placeholder actor since we only need its actor descriptor.
 					WorldPartitionContext->GetWorld()->DestroyActor(CustomHLODPlaceholderActor, false, false);
-					TrashObject(CustomHLODPlaceholderActor);
+					FName UniqueTrashName = MakeUniqueObjectName(CustomHLODPlaceholderActor->GetOuter(), CustomHLODPlaceholderActor->GetClass(), NAME_TrashedPackage);
+					CustomHLODPlaceholderActor->Rename(*UniqueTrashName.ToString(), nullptr, REN_DontCreateRedirectors | REN_NonTransactional | REN_DoNotDirty);
 				}
 
 				// If Custom HLOD actor is used with some HLOD layer different than CustomHLODActor layer don't register its actor descriptor
