[AutoRTFM] Fix inverted logic on #else comment.

diff --git a/Engine/Source/Programs/AutoRTFMTests/Private/Catch2Includes.h b/Engine/Source/Programs/AutoRTFMTests/Private/Catch2Includes.h
--- a/Engine/Source/Programs/AutoRTFMTests/Private/Catch2Includes.h
+++ b/Engine/Source/Programs/AutoRTFMTests/Private/Catch2Includes.h
@@ -88,7 +88,7 @@ THIRD_PARTY_INCLUDES_END
     #define AUTORTFM_CATCH_TEST_CASE(TEST_NAME, ...) \
         AUTORTFM_DISABLE static void TEST_NAME(); \
         INTERNAL_CATCH_TESTCASE2(TEST_NAME)
-#else  // ^^ !CATCH_CONFIG_EXPERIMENTAL_STATIC_ANALYSIS_SUPPORT | vv CATCH_CONFIG_EXPERIMENTAL_STATIC_ANALYSIS_SUPPORT
+#else  // ^^ CATCH_CONFIG_EXPERIMENTAL_STATIC_ANALYSIS_SUPPORT | vv !CATCH_CONFIG_EXPERIMENTAL_STATIC_ANALYSIS_SUPPORT
     #define AUTORTFM_CATCH_TEST_CASE(TEST_NAME, ...) \
         AUTORTFM_DISABLE static void TEST_NAME(); \
         INTERNAL_CATCH_TESTCASE2(Catch::DefaultGroup, TEST_NAME, __VA_ARGS__)
