Use a singlethreaded `TQueue` in `FArchiveObjectCrc32` as its not a threadsafe class anyway and blew up in transactions.

diff --git a/Engine/Source/Programs/AutoRTFMTests/Private/UECoreTests.cpp b/Engine/Source/Programs/AutoRTFMTests/Private/UECoreTests.cpp
--- a/Engine/Source/Programs/AutoRTFMTests/Private/UECoreTests.cpp
+++ b/Engine/Source/Programs/AutoRTFMTests/Private/UECoreTests.cpp
@@ -22,6 +22,7 @@
 #include "Logging/StructuredLogFormat.h"
 #include "Memory/VirtualStackAllocator.h"
 #include "Modules/ModuleManager.h"
+#include "Serialization/ArchiveObjectCrc32.h"
 #include "Serialization/CustomVersion.h"
 #include "UObject/CoreRedirects.h"
 #include "UObject/DynamicallyTypedValue.h"
@@ -2352,3 +2353,25 @@ TEST_CASE("UECore.UnifiedError")
 		});
 	}
 }
+
+TEST_CASE("UECore.ArchiveObjectCrc32")
+{
+	SECTION("Constructed in a transaction")
+	{
+		AutoRTFM::Testing::Commit([&]
+		{
+			FArchiveObjectCrc32 Archive;
+			Archive.Crc32(NewObject<UMyAutoRTFMTestObject>());
+		});
+	}
+
+	SECTION("Constructed outwith a transaction")
+	{
+		FArchiveObjectCrc32 Archive;
+
+		AutoRTFM::Testing::Commit([&]
+		{
+			Archive.Crc32(NewObject<UMyAutoRTFMTestObject>());
+		});
+	}
+}
diff --git a/Engine/Source/Runtime/CoreUObject/Public/Serialization/ArchiveObjectCrc32.h b/Engine/Source/Runtime/CoreUObject/Public/Serialization/ArchiveObjectCrc32.h
--- a/Engine/Source/Runtime/CoreUObject/Public/Serialization/ArchiveObjectCrc32.h
+++ b/Engine/Source/Runtime/CoreUObject/Public/Serialization/ArchiveObjectCrc32.h
@@ -52,7 +52,7 @@ protected:
 	/** Internal archive used for serialization */
 	FMemoryWriter MemoryWriter;
 	/** Internal queue of object references awaiting serialization */
-	TQueue<UObject*> ObjectsToSerialize;
+	TQueue<UObject*, EQueueMode::SingleThreaded> ObjectsToSerialize;
 	/** Internal currently serialized object */
 	const UObject* ObjectBeingSerialized;
 	/** Internal root object */
