[AutoRTFM] Improve AutoRTFM error reporting when freeing transactional memory from the open.

This was another case where we properly catch the error, but did a terrible job of conveying the problem to the developer.

diff --git a/Engine/Source/Programs/AutoRTFMTests/Private/Abort.cpp b/Engine/Source/Programs/AutoRTFMTests/Private/Abort.cpp
--- a/Engine/Source/Programs/AutoRTFMTests/Private/Abort.cpp
+++ b/Engine/Source/Programs/AutoRTFMTests/Private/Abort.cpp
@@ -1442,3 +1442,29 @@ TEST_CASE("Abort.Language.ErrorHints.TryAcquireSRWLockShared")
 	// Did we print our message?
 	REQUIRE(WarningContext.HasWarningSubstring("A read-write lock was accessed while inside an AutoRTFM transaction."));
 }
+
+#if 0
+// Testing an open free is flaky, since it leads to a double free--one inside the open block, and one during the
+// transactional abort caused by an AutoRTFM hazard. This can cause the unit test harness to fail on a subsequent test,
+// or during shutdown. In normal operation, this doesn't matter, because AutoRTFM hazards are designed to abort.
+TEST_CASE("Abort.Language.ErrorHints.MallocOpenFree")
+{
+	FScopedErrorHintTest ScopedTest;
+    AutoRTFMTestUtils::FCaptureWarningContext WarningContext;
+
+	const AutoRTFM::ETransactionResult Result = AutoRTFM::Transact([&]
+	{
+		void* Data = FMemory::Malloc(100);
+
+		AutoRTFM::Open([&]
+		{
+			FMemory::Free(Data);
+		});
+	});
+
+	REQUIRE(AutoRTFM::ETransactionResult::AbortedByLanguage == Result);
+
+	// Did we print our message?
+	REQUIRE(WarningContext.HasWarningSubstring("Memory has been freed from the open while inside an AutoRTFM transaction."));
+}
+#endif
\ No newline at end of file
diff --git a/Engine/Source/Runtime/AutoRTFM/Private/ReportAutoRTFMHazard.cpp b/Engine/Source/Runtime/AutoRTFM/Private/ReportAutoRTFMHazard.cpp
--- a/Engine/Source/Runtime/AutoRTFM/Private/ReportAutoRTFMHazard.cpp
+++ b/Engine/Source/Runtime/AutoRTFM/Private/ReportAutoRTFMHazard.cpp
@@ -39,6 +39,17 @@ read-write locks, but it does support transactionally-safe read-write locks that
 Use the call stack below to find the FRWLock accessed immediately before this error occurred, and
 replace it with FTransactionallySafeFRWLock. If FReadScopeLock, FWriteScopeLock or FRWScopeLock are
 used, replace them with UE::TReadScopeLock, UE::TWriteScopeLock or UE::TRWScopeLock.
+)";
+			break;
+
+		case EHazardType::OpenFree:
+			Message = R"(
+Memory has been freed from the open while inside an AutoRTFM transaction. This memory was allocated
+or written-to from the closed within the same transaction; freeing it makes rollback impossible.
+
+Use the call stack below to find the location of the free call; the open block will be further up in
+the stack. If possible, avoid going into the open at all; if it is unavoidable, ensure that memory
+is never modified or freed from the open if it might also be written-to from the closed.
 )";
 			break;
 
diff --git a/Engine/Source/Runtime/AutoRTFM/Private/ReportAutoRTFMHazard.h b/Engine/Source/Runtime/AutoRTFM/Private/ReportAutoRTFMHazard.h
--- a/Engine/Source/Runtime/AutoRTFM/Private/ReportAutoRTFMHazard.h
+++ b/Engine/Source/Runtime/AutoRTFM/Private/ReportAutoRTFMHazard.h
@@ -13,6 +13,7 @@ enum class EHazardType
 {
 	CriticalSection,
 	RWLock,
+	OpenFree,
 };
 
 AUTORTFM_DISABLE void ReportAutoRTFMHazard(EHazardType HazardType);
diff --git a/Engine/Source/Runtime/AutoRTFM/Private/TransactionInlines.h b/Engine/Source/Runtime/AutoRTFM/Private/TransactionInlines.h
--- a/Engine/Source/Runtime/AutoRTFM/Private/TransactionInlines.h
+++ b/Engine/Source/Runtime/AutoRTFM/Private/TransactionInlines.h
@@ -10,6 +10,7 @@
 #include "Transaction.h"
 #include "Utils.h"
 #include "WriteLog.h"
+#include "ReportAutoRTFMHazard.h"
 
 #include <utility>
 
@@ -184,7 +185,11 @@ UE_AUTORTFM_FORCEINLINE void FTransaction::DidFree(void* LogicalAddress)
 	if (!bIsInAllocateFn)
 	{
 		AutoRTFM::TScopedGuard<bool> RecursionGuard(bIsInAllocateFn, true);
-		AUTORTFM_ASSERT(!NewMemoryTracker.Contains(LogicalAddress, 1));
+		if (NewMemoryTracker.Contains(LogicalAddress, 1))
+		{
+			ForTheRuntime::ReportAutoRTFMHazard(ForTheRuntime::EHazardType::OpenFree);
+			AUTORTFM_REPORT_ERROR("Transactional memory was freed from the open.");
+		}
 	}
 }
 
