[AutoRTFM] Add FTransactionallySafeSharedRecursiveMutex

Required by Chaos to fix a hang when re-entering a transactionally safe mutex, outside of a transaction.

Reworked the AutoRTFM MutexTests to test all the different flavours of mutex we have (we really need to consolidate these implementations).

diff --git a/Engine/Source/Programs/AutoRTFMTests/Private/MutexTests.cpp b/Engine/Source/Programs/AutoRTFMTests/Private/MutexTests.cpp
--- a/Engine/Source/Programs/AutoRTFMTests/Private/MutexTests.cpp
+++ b/Engine/Source/Programs/AutoRTFMTests/Private/MutexTests.cpp
@@ -1,10 +1,13 @@
 // Copyright Epic Games, Inc. All Rights Reserved.
 
+#include "Async/SharedLock.h"
 #include "AutoRTFM.h"
 #include "AutoRTFM/Testing.h"
 #include "AutoRTFMTestUtils.h"
 #include "Catch2Includes.h"
 #include "Async/TransactionallySafeMutex.h"
+#include "Async/TransactionallySafeSharedMutex.h"
+#include "Async/TransactionallySafeSharedRecursiveMutex.h"
 
 #include <atomic>
 #include <condition_variable>
@@ -12,887 +15,1325 @@
 #include <memory>
 #include <mutex>
 #include <thread>
+#include <type_traits>
 
-TEST_CASE("TransactionallySafeMutex.Outside Transaction")
+template<typename T>
+concept CHasIsLockedMethod = requires(T Object) { { Object.IsLocked() } -> std::same_as<bool>; };
+
+template<typename T>
+concept CHasIsLockSharedMethod = requires(T Object) { { Object.IsLockShared() } -> std::same_as<bool>; };
+
+TEMPLATE_TEST_CASE("MutexTests", "", 
+	UE::FTransactionallySafeMutex,
+	UE::FTransactionallySafeRecursiveMutex,
+	UE::FTransactionallySafeSharedMutex,
+	UE::FTransactionallySafeSharedRecursiveMutex)
 {
-	UE::FTransactionallySafeMutex Mutex;
+	using MutexType = TestType;
+	
+	// TODO(SOL-7661): Unify the transactionally-safe mutex implementations so
+	// we can have a consistent set of behaviors across all variants.
 
-	AutoRTFM::Testing::Abort([&]
-	{
-		UE::TScopeLock Lock(Mutex);
-		AutoRTFM::AbortTransaction();
-	});
+	// Some mutex implementations do not currently support having an unbalanced
+	// number of calls to Lock() / Unlock().
+	constexpr bool bSupportsUnbalancedLockingInTransaction = std::is_same_v<MutexType, UE::FTransactionallySafeMutex>;
 
-	AutoRTFM::Testing::Commit([&]
-	{
-		UE::TScopeLock Lock(Mutex);
-	});
-}
+	// Some mutex implementations do not currently support locking and unlocking
+	// across open / closed boundaries.
+	constexpr bool bSupportsMixedOpenCloseLocking = std::is_same_v<MutexType, UE::FTransactionallySafeMutex>;
 
-TEST_CASE("TransactionallySafeMutex.IsLocked Outside Transaction")
-{
-	UE::FTransactionallySafeMutex Mutex;
+	// True if the mutex supports TryLock() within a transaction
+	constexpr bool bSupportsTryLockInTransaction =
+		std::is_same_v<MutexType, UE::FTransactionallySafeMutex> ||
+		std::is_same_v<MutexType, UE::FTransactionallySafeSharedMutex>;
 
-	REQUIRE(!Mutex.IsLocked());
-	Mutex.Lock();
-	REQUIRE(Mutex.IsLocked());
-	Mutex.Unlock();
-	REQUIRE(!Mutex.IsLocked());
-}
+	// Some mutex implementations do not currently support the IsLocked() method
+	constexpr bool bSupportsIsLocked = CHasIsLockedMethod<MutexType>;
 
-TEST_CASE("TransactionallySafeMutex.Inside Transaction")
-{
-	AutoRTFM::Testing::Abort([&]
-	{
-		UE::FTransactionallySafeMutex Mutex;
-		UE::TScopeLock Lock(Mutex);
-		AutoRTFM::AbortTransaction();
-	});
+	// True if the mutex has a IsLockShared() method.
+	constexpr bool bSupportsIsLockShared = CHasIsLockSharedMethod<MutexType>;
 
-	AutoRTFM::Testing::Commit([&]
-	{
-		UE::FTransactionallySafeMutex Mutex;
-		UE::TScopeLock Lock(Mutex);
-	});
-}
+	// True if the mutex supports recursive locking
+	constexpr bool bIsRecursive = std::is_same_v<MutexType, UE::FTransactionallySafeSharedRecursiveMutex>;
 
-TEST_CASE("TransactionallySafeMutex.IsLocked Inside Transaction")
-{
-	UE::FTransactionallySafeMutex Mutex;
-	REQUIRE(!Mutex.IsLocked());
+	// True if the mutex is shared
+	constexpr bool bIsShared =
+		std::is_same_v<MutexType, UE::FTransactionallySafeSharedMutex> ||
+		std::is_same_v<MutexType, UE::FTransactionallySafeSharedRecursiveMutex>;
 
-	AutoRTFM::Testing::Commit([&]
+	auto RequireLocked = [](MutexType& Mutex)
 	{
-		REQUIRE(!Mutex.IsLocked());
-		Mutex.Lock();
-		REQUIRE(Mutex.IsLocked());
-		Mutex.Unlock();
-		REQUIRE(Mutex.IsLocked());
-	});
-
-	REQUIRE(!Mutex.IsLocked());
-}
+		if constexpr (bSupportsIsLocked)
+		{
+			REQUIRE(Mutex.IsLocked());
+		}
+		else
+		{
+			// No way to probe the mutex locked state.
+			// As a best effort attempt to unlock and re-lock.
+			Mutex.Unlock();
+			Mutex.Lock();
+		}
+	};
+	
+	auto RequireNotLocked = [](MutexType& Mutex)
+	{
+		if constexpr (bSupportsIsLocked)
+		{
+			REQUIRE(!Mutex.IsLocked());
+		}
+		else
+		{
+			// No way to probe the mutex locked state.
+			// As a best effort attempt to lock and unlock.
+			Mutex.Lock();
+			Mutex.Unlock();
+		}
+	};
 
-TEST_CASE("TransactionallySafeMutex.Inside Transaction Used In Nested Transaction")
-{
-	AutoRTFM::Testing::Abort([&]
+	auto RequireLockShared = [](MutexType& Mutex)
+	{
+		if constexpr (bSupportsIsLockShared)
+		{
+			REQUIRE(Mutex.IsLockShared());
+		}
+	};
+	
+	auto RequireNotSharedLocked = [](MutexType& Mutex)
+	{
+		if constexpr (bSupportsIsLockShared)
+		{
+			REQUIRE(!Mutex.IsLockShared());
+		}
+	};
+	
+	SECTION("Construct Outside Transaction")
 	{
-		UE::FTransactionallySafeMutex Mutex;
+		MutexType Mutex;
 
 		AutoRTFM::Testing::Abort([&]
 		{
 			UE::TScopeLock Lock(Mutex);
-			AutoRTFM::CascadingAbortTransaction();
+			AutoRTFM::AbortTransaction();
 		});
-	});
 
-	AutoRTFM::Testing::Commit([&]
-	{
-		UE::FTransactionallySafeMutex Mutex;
+		AutoRTFM::Testing::Commit([&]
+		{
+			UE::TScopeLock Lock(Mutex);
+		});
+	}
 
+	SECTION("Construct Inside Transaction")
+	{
 		AutoRTFM::Testing::Abort([&]
 		{
+			MutexType Mutex;
 			UE::TScopeLock Lock(Mutex);
 			AutoRTFM::AbortTransaction();
 		});
-	});
-
-	AutoRTFM::Testing::Abort([&]
-	{
-		UE::FTransactionallySafeMutex Mutex;
 
 		AutoRTFM::Testing::Commit([&]
 		{
+			MutexType Mutex;
 			UE::TScopeLock Lock(Mutex);
 		});
+	}
+
+	SECTION("IsLocked Outside Transaction")
+	{
+		MutexType Mutex;
 
-		AutoRTFM::AbortTransaction();
-	});
+		RequireNotLocked(Mutex);
+		Mutex.Lock();
+		RequireLocked(Mutex);
+		Mutex.Unlock();
+		RequireNotLocked(Mutex);
+	}
 
-	AutoRTFM::Testing::Commit([&]
+	SECTION("IsLocked Inside Transaction")
 	{
-		UE::FTransactionallySafeMutex Mutex;
+		MutexType Mutex;
+		RequireNotLocked(Mutex);
 
 		AutoRTFM::Testing::Commit([&]
 		{
-			UE::TScopeLock Lock(Mutex);
+			RequireNotLocked(Mutex);
+			Mutex.Lock();
+			RequireLocked(Mutex);
+			Mutex.Unlock();
+			RequireLocked(Mutex);
 		});
-	});
-}
 
-TEST_CASE("TransactionallySafeMutex.In Static Local Initializer")
-{
-	struct AUTORTFM_ENABLE MyStruct final
-	{
-		UE::FTransactionallySafeMutex Mutex;
-	};
+		RequireNotLocked(Mutex);
+	}
 
-	auto Lambda = [] AUTORTFM_ENABLE
+	SECTION("Inside Transaction Used In Nested Transaction")
 	{
-		static MyStruct Mine;
-		UE::TScopeLock _(Mine.Mutex);
-		return 42;
-	};
+		AutoRTFM::Testing::Abort([&]
+		{
+			MutexType Mutex;
 
-	AutoRTFM::Testing::Abort([&]
-	{
-		REQUIRE(42 == Lambda());
-		AutoRTFM::AbortTransaction();
-	});
+			AutoRTFM::Testing::Abort([&]
+			{
+				UE::TScopeLock Lock(Mutex);
+				AutoRTFM::CascadingAbortTransaction();
+			});
+		});
 
-	REQUIRE(42 == Lambda());
+		AutoRTFM::Testing::Commit([&]
+		{
+			MutexType Mutex;
 
-	AutoRTFM::Testing::Commit([&]
-	{
-		REQUIRE(42 == Lambda());
-	});
+			AutoRTFM::Testing::Abort([&]
+			{
+				UE::TScopeLock Lock(Mutex);
+				AutoRTFM::AbortTransaction();
+			});
+		});
 
-	REQUIRE(42 == Lambda());
-}
+		AutoRTFM::Testing::Abort([&]
+		{
+			MutexType Mutex;
 
-TEST_CASE("TransactionallySafeMutex.In Static Local Initializer Called From Open")
-{
-	struct AUTORTFM_ENABLE MyStruct final
-	{
-		UE::FTransactionallySafeMutex Mutex;
-	};
+			AutoRTFM::Testing::Commit([&]
+			{
+				UE::TScopeLock Lock(Mutex);
+			});
 
-	auto Lambda = [] AUTORTFM_ENABLE
-	{
-		static MyStruct Mine;
-		UE::TScopeLock _(Mine.Mutex);
-		return 42;
-	};
+			AutoRTFM::AbortTransaction();
+		});
 
-	AutoRTFM::Testing::Abort([&]
-	{
-		AutoRTFM::Open([&] { REQUIRE(42 == Lambda()); });
-		AutoRTFM::AbortTransaction();
-	});
+		AutoRTFM::Testing::Commit([&]
+		{
+			MutexType Mutex;
 
-	REQUIRE(42 == Lambda());
+			AutoRTFM::Testing::Commit([&]
+			{
+				UE::TScopeLock Lock(Mutex);
+			});
+		});
+	}
 
-	AutoRTFM::Testing::Commit([&]
+	SECTION("In Static Local Initializer")
 	{
-		AutoRTFM::Open([&] { REQUIRE(42 == Lambda()); });
-	});
-
-	REQUIRE(42 == Lambda());
-}
+		struct AUTORTFM_ENABLE MyStruct final
+		{
+			MutexType Mutex;
+		};
 
-TEST_CASE("TransactionallySafeMutex.Delete Heap Allocated Mutex")
-{
-	SECTION("SingleThread") // Mutex owned and destructed by this thread
-	{
-		UE::FTransactionallySafeMutex* const Mutex = new UE::FTransactionallySafeMutex();
+		auto Lambda = [] AUTORTFM_ENABLE
+		{
+			static MyStruct Mine;
+			UE::TScopeLock _(Mine.Mutex);
+			return 42;
+		};
 
 		AutoRTFM::Testing::Abort([&]
 		{
-			Mutex->Lock();
-			delete Mutex;
+			REQUIRE(42 == Lambda());
 			AutoRTFM::AbortTransaction();
 		});
 
-		REQUIRE(!Mutex->IsLocked());
+		REQUIRE(42 == Lambda());
 
 		AutoRTFM::Testing::Commit([&]
 		{
-			Mutex->Lock();
-			delete Mutex;
+			REQUIRE(42 == Lambda());
 		});
+
+		REQUIRE(42 == Lambda());
 	}
-		
-	SECTION("MultiThread") // Mutex owned and destructed by another thread
+
+	SECTION("In Static Local Initializer Called From Open")
 	{
-		// This test does not support retries due to coordination with another thread.
-		AutoRTFMTestUtils::FScopedRetry NoRetry(AutoRTFM::ForTheRuntime::EAutoRTFMRetryTransactionState::NoRetry);
+		struct AUTORTFM_ENABLE MyStruct final
+		{
+			MutexType Mutex;
+		};
 
-		enum EState
+		auto Lambda = [] AUTORTFM_ENABLE
 		{
-			InitializeThread,
-			MutexReady,
-			MutexUsed,
-			MutexDeleted,
+			static MyStruct Mine;
+			UE::TScopeLock _(Mine.Mutex);
+			return 42;
 		};
 
-		struct AUTORTFM_OPEN FEvent
+		AutoRTFM::Testing::Abort([&]
 		{
-			void Signal(EState State)
-			{
-				{
-					std::unique_lock Lock(Mutex);
-					CurrentState = State;
-				}
-				CV.notify_one();
-			}
+			AutoRTFM::Open([&] { REQUIRE(42 == Lambda()); });
+			AutoRTFM::AbortTransaction();
+		});
 
-			void Wait(EState State)
-			{
-				std::unique_lock Lock(Mutex);
-				CV.wait(Lock, [this, State] { return CurrentState == State; });
-			}
-		private:
-			EState CurrentState = InitializeThread;
-			std::mutex Mutex;
-			std::condition_variable CV;
-		};
-			
-		UE::FTransactionallySafeMutex* Mutex = nullptr;
-		FEvent Event;
+		REQUIRE(42 == Lambda());
 
-		std::thread Thread([&Mutex, &Event] AUTORTFM_OPEN
+		AutoRTFM::Testing::Commit([&]
 		{
-			Mutex = new UE::FTransactionallySafeMutex();
-			Event.Signal(MutexReady);
-			Event.Wait(MutexUsed);
-			delete Mutex;
-			Event.Signal(MutexDeleted);
+			AutoRTFM::Open([&] { REQUIRE(42 == Lambda()); });
 		});
 
-		Event.Wait(MutexReady);
+		REQUIRE(42 == Lambda());
+	}
 
-		SECTION("Lock, Abort, Destroy")
+	SECTION("Delete Heap Allocated Mutex")
+	{
+		SECTION("SingleThread") // Mutex owned and destructed by this thread
 		{
+			UE::FTransactionallySafeMutex* const Mutex = new UE::FTransactionallySafeMutex();
+
 			AutoRTFM::Testing::Abort([&]
 			{
 				Mutex->Lock();
+				delete Mutex;
 				AutoRTFM::AbortTransaction();
 			});
 
 			REQUIRE(!Mutex->IsLocked());
 
-			Event.Signal(MutexUsed);
-			Event.Wait(MutexDeleted);
+			AutoRTFM::Testing::Commit([&]
+			{
+				Mutex->Lock();
+				delete Mutex;
+			});
 		}
-
-		SECTION("Lock, Unlock, Destroy, Abort")
+			
+		SECTION("MultiThread") // Mutex owned and destructed by another thread
 		{
-			AutoRTFM::Testing::Abort([&]
+			// This test does not support retries due to coordination with another thread.
+			AutoRTFMTestUtils::FScopedRetry NoRetry(AutoRTFM::ForTheRuntime::EAutoRTFMRetryTransactionState::NoRetry);
+
+			enum EState
 			{
-				Mutex->Lock();
-				Mutex->Unlock();
+				InitializeThread,
+				MutexReady,
+				MutexUsed,
+				MutexDeleted,
+			};
 
-				Event.Signal(MutexUsed);
-				Event.Wait(MutexDeleted);
+			struct AUTORTFM_OPEN FEvent
+			{
+				void Signal(EState State)
+				{
+					{
+						std::unique_lock Lock(Mutex);
+						CurrentState = State;
+					}
+					CV.notify_one();
+				}
 
-				AutoRTFM::AbortTransaction();
+				void Wait(EState State)
+				{
+					std::unique_lock Lock(Mutex);
+					CV.wait(Lock, [this, State] { return CurrentState == State; });
+				}
+			private:
+				EState CurrentState = InitializeThread;
+				std::mutex Mutex;
+				std::condition_variable CV;
+			};
+				
+			UE::FTransactionallySafeMutex* Mutex = nullptr;
+			FEvent Event;
+
+			std::thread Thread([&Mutex, &Event] AUTORTFM_OPEN
+			{
+				Mutex = new UE::FTransactionallySafeMutex();
+				Event.Signal(MutexReady);
+				Event.Wait(MutexUsed);
+				delete Mutex;
+				Event.Signal(MutexDeleted);
 			});
-		}
 
-		SECTION("Lock, Unlock, Destroy, Commit")
-		{
-			AutoRTFM::Testing::Commit([&]
+			Event.Wait(MutexReady);
+
+			SECTION("Lock, Abort, Destroy")
 			{
-				Mutex->Lock();
-				Mutex->Unlock();
+				AutoRTFM::Testing::Abort([&]
+				{
+					Mutex->Lock();
+					AutoRTFM::AbortTransaction();
+				});
+
+				REQUIRE(!Mutex->IsLocked());
 
 				Event.Signal(MutexUsed);
 				Event.Wait(MutexDeleted);
-			});
-		}
+			}
 
-		Thread.join();
-	}
-}
+			SECTION("Lock, Unlock, Destroy, Abort")
+			{
+				AutoRTFM::Testing::Abort([&]
+				{
+					Mutex->Lock();
+					Mutex->Unlock();
 
-TEST_CASE("TransactionallySafeMutex.Lock Within, Unlock Outside")
-{
-	UE::FTransactionallySafeMutex Mutex;
+					Event.Signal(MutexUsed);
+					Event.Wait(MutexDeleted);
 
-	AutoRTFM::Testing::Commit([&]
-	{
-		Mutex.Lock();
-	});
+					AutoRTFM::AbortTransaction();
+				});
+			}
 
-	REQUIRE(Mutex.IsLocked());
-	Mutex.Unlock();
-	REQUIRE(!Mutex.IsLocked());
-}
+			SECTION("Lock, Unlock, Destroy, Commit")
+			{
+				AutoRTFM::Testing::Commit([&]
+				{
+					Mutex->Lock();
+					Mutex->Unlock();
 
-TEST_CASE("TransactionallySafeMutex.TryLock Within, Unlock Outside")
-{
-	UE::FTransactionallySafeMutex Mutex;
+					Event.Signal(MutexUsed);
+					Event.Wait(MutexDeleted);
+				});
+			}
+
+			Thread.join();
+		}
+	}
 
-	AutoRTFM::Testing::Commit([&]
+	if constexpr (bSupportsUnbalancedLockingInTransaction)
 	{
-		REQUIRE(Mutex.TryLock());
-	});
+		SECTION("Lock Within, Unlock Outside")
+		{
+			MutexType Mutex;
 
-	REQUIRE(Mutex.IsLocked());
-	Mutex.Unlock();
-	REQUIRE(!Mutex.IsLocked());
-}
+			AutoRTFM::Testing::Commit([&]
+			{
+				Mutex.Lock();
+			});
 
-TEST_CASE("TransactionallySafeMutex.Lock Outside, Unlock Within")
-{
-	UE::FTransactionallySafeMutex Mutex = UE::FTransactionallySafeMutex(UE::FAcquireLock());
+			RequireLocked(Mutex);
+			Mutex.Unlock();
+			RequireNotLocked(Mutex);
+		}
+	}
 
-	AutoRTFM::Testing::Commit([&]
+	if constexpr (bSupportsUnbalancedLockingInTransaction && bSupportsTryLockInTransaction)
 	{
-		REQUIRE(Mutex.IsLocked());
-		Mutex.Unlock();
-	});
+		SECTION("TryLock Within, Unlock Outside")
+		{
+			MutexType Mutex;
 
-	REQUIRE(!Mutex.IsLocked());
-}
+			AutoRTFM::Testing::Commit([&]
+			{
+				REQUIRE(Mutex.TryLock());
+			});
 
-TEST_CASE("TransactionallySafeMutex.Lock, Unlock, Lock")
-{
-	UE::FTransactionallySafeMutex Mutex;
+			RequireLocked(Mutex);
+			Mutex.Unlock();
+			RequireNotLocked(Mutex);
+		}
+	}
 
-	SECTION("Commit")
+	if constexpr (bSupportsUnbalancedLockingInTransaction)
 	{
-		AutoRTFM::Testing::Commit([&]
+		SECTION("Lock Outside, Unlock Within")
 		{
+			MutexType Mutex;
 			Mutex.Lock();
-			Mutex.Unlock();
-			Mutex.Lock();
-		});
 
-		REQUIRE(Mutex.IsLocked());
-		Mutex.Unlock();
+			AutoRTFM::Testing::Commit([&]
+			{
+				RequireLocked(Mutex);
+				Mutex.Unlock();
+			});
+
+			RequireNotLocked(Mutex);
+		}
 	}
 
-	SECTION("Abort")
+	if constexpr (bSupportsUnbalancedLockingInTransaction)
 	{
-		AutoRTFM::Testing::Abort([&]
+		SECTION("Lock, Unlock, Lock")
 		{
-			Mutex.Lock();
-			Mutex.Unlock();
-			Mutex.Lock();
-			AutoRTFM::AbortTransaction();
-		});
+			MutexType Mutex;
 
-		REQUIRE(!Mutex.IsLocked());
-	}
-}
+			SECTION("Commit")
+			{
+				AutoRTFM::Testing::Commit([&]
+				{
+					Mutex.Lock();
+					Mutex.Unlock();
+					Mutex.Lock();
+				});
 
-TEST_CASE("TransactionallySafeMutex.TryLock, Unlock, TryLock")
-{
-	UE::FTransactionallySafeMutex Mutex;
+				RequireLocked(Mutex);
+				Mutex.Unlock();
+			}
 
-	SECTION("Commit")
-	{
-		AutoRTFM::Testing::Commit([&]
-		{
-			REQUIRE(Mutex.TryLock());
-			Mutex.Unlock();
-			REQUIRE(Mutex.TryLock());
-		});
+			SECTION("Abort")
+			{
+				AutoRTFM::Testing::Abort([&]
+				{
+					Mutex.Lock();
+					Mutex.Unlock();
+					Mutex.Lock();
+					AutoRTFM::AbortTransaction();
+				});
 
-		REQUIRE(Mutex.IsLocked());
-		Mutex.Unlock();
+				RequireNotLocked(Mutex);
+			}
+		}
 	}
-		
-	SECTION("Abort")
+
+	if constexpr (bSupportsUnbalancedLockingInTransaction && bSupportsTryLockInTransaction)
 	{
-		AutoRTFM::Testing::Abort([&]
+		SECTION("TryLock, Unlock, TryLock")
 		{
-			REQUIRE(Mutex.TryLock());
-			Mutex.Unlock();
-			REQUIRE(Mutex.TryLock());
-			AutoRTFM::AbortTransaction();
-		});
+			MutexType Mutex;
 
-		REQUIRE(!Mutex.IsLocked());
-	}
-}
+			SECTION("Commit")
+			{
+				AutoRTFM::Testing::Commit([&]
+				{
+					REQUIRE(Mutex.TryLock());
+					Mutex.Unlock();
+					REQUIRE(Mutex.TryLock());
+				});
 
-TEST_CASE("TransactionallySafeMutex.Unlock, Lock, Unlock")
-{
-	UE::FTransactionallySafeMutex Mutex;
-	Mutex.Lock();
+				RequireLocked(Mutex);
+				Mutex.Unlock();
+			}
+				
+			SECTION("Abort")
+			{
+				AutoRTFM::Testing::Abort([&]
+				{
+					REQUIRE(Mutex.TryLock());
+					Mutex.Unlock();
+					REQUIRE(Mutex.TryLock());
+					AutoRTFM::AbortTransaction();
+				});
+
+				RequireNotLocked(Mutex);
+			}
+		}
+	}
 
-	SECTION("Commit")
+	if constexpr (bSupportsUnbalancedLockingInTransaction)
 	{
-		AutoRTFM::Testing::Commit([&]
+		SECTION("Unlock, Lock, Unlock")
 		{
-			Mutex.Unlock();
+			MutexType Mutex;
 			Mutex.Lock();
-			Mutex.Unlock();
-		});
-
-		REQUIRE(!Mutex.IsLocked());
-	}
 
-	SECTION("Commit")
+			SECTION("Commit")
+			{
+				AutoRTFM::Testing::Commit([&]
+				{
+					Mutex.Unlock();
+					Mutex.Lock();
+					Mutex.Unlock();
+				});
+
+				RequireNotLocked(Mutex);
+			}
+
+			SECTION("Commit")
+			{
+				AutoRTFM::Testing::Abort([&]
+				{
+					Mutex.Unlock();
+					Mutex.Lock();
+					Mutex.Unlock();
+					AutoRTFM::AbortTransaction();
+				});
+
+				RequireLocked(Mutex);
+				Mutex.Unlock();
+			}
+		}
+	}
+
+	if constexpr (bSupportsUnbalancedLockingInTransaction && bSupportsTryLockInTransaction)
 	{
-		AutoRTFM::Testing::Abort([&]
+		SECTION("Unlock, TryLock, Unlock")
 		{
-			Mutex.Unlock();
+			MutexType Mutex;
 			Mutex.Lock();
-			Mutex.Unlock();
-			AutoRTFM::AbortTransaction();
-		});
 
-		REQUIRE(Mutex.IsLocked());
-		Mutex.Unlock();
-	}
-}
+			SECTION("Commit")
+			{
+				AutoRTFM::Testing::Commit([&]
+				{
+					Mutex.Unlock();
+					REQUIRE(Mutex.TryLock());
+					Mutex.Unlock();
+				});
 
-TEST_CASE("TransactionallySafeMutex.Unlock, TryLock, Unlock")
-{
-	UE::FTransactionallySafeMutex Mutex;
-	Mutex.Lock();
+				RequireNotLocked(Mutex);
+			}
+
+			SECTION("Abort")
+			{
+				AutoRTFM::Testing::Abort([&]
+				{
+					Mutex.Unlock();
+					REQUIRE(Mutex.TryLock());
+					Mutex.Unlock();
+					AutoRTFM::AbortTransaction();
+				});
+
+				RequireLocked(Mutex);
+				Mutex.Unlock();
+			}
+		}
+	}
 
-	SECTION("Commit")
+	if constexpr (bSupportsUnbalancedLockingInTransaction)
 	{
-		AutoRTFM::Testing::Commit([&]
+		SECTION("Lock Outside, Unlock & Lock Within")
 		{
-			Mutex.Unlock();
-			REQUIRE(Mutex.TryLock());
-			Mutex.Unlock();
-		});
+			MutexType Mutex;
+			Mutex.Lock();
+
+			AutoRTFM::Testing::Commit([&]
+			{
+				RequireLocked(Mutex);
+				Mutex.Unlock();
+				Mutex.Lock();
+			});
 
-		REQUIRE(!Mutex.IsLocked());
+			RequireLocked(Mutex);
+		}
 	}
 
-	SECTION("Commit")
+	if constexpr (bSupportsUnbalancedLockingInTransaction)
 	{
-		AutoRTFM::Testing::Abort([&]
+		SECTION("Commit(Lock, Commit(Unlock))")
 		{
-			Mutex.Unlock();
-			REQUIRE(Mutex.TryLock());
-			Mutex.Unlock();
-			AutoRTFM::AbortTransaction();
-		});
+			MutexType Mutex;
 
-		REQUIRE(Mutex.IsLocked());
-		Mutex.Unlock();
-	}
-}
-
-TEST_CASE("TransactionallySafeMutex.Lock Outside, Unlock & Lock Within")
-{
-	UE::FTransactionallySafeMutex Mutex = UE::FTransactionallySafeMutex(UE::FAcquireLock());
+			AutoRTFM::Testing::Commit([&]
+			{
+				Mutex.Lock();
 
-	AutoRTFM::Testing::Commit([&]
-	{
-		REQUIRE(Mutex.IsLocked());
-		Mutex.Unlock();
-		Mutex.Lock();
-	});
+				AutoRTFM::Testing::Commit([&]
+				{
+					Mutex.Unlock();
+				});
 
-	REQUIRE(Mutex.IsLocked());
-}
+				RequireLocked(Mutex);
+			});
 
-TEST_CASE("TransactionallySafeMutex.Commit(Lock, Commit(Unlock))")
-{
-	UE::FTransactionallySafeMutex Mutex;
+			RequireNotLocked(Mutex);
+		}
+	}
 
-	AutoRTFM::Testing::Commit([&]
+	if constexpr (bSupportsUnbalancedLockingInTransaction)
 	{
-		Mutex.Lock();
-
-		AutoRTFM::Testing::Commit([&]
+		SECTION("Abort(Lock, Commit(Unlock, Lock))")
 		{
-			Mutex.Unlock();
-		});
+			MutexType Mutex;
 
-		REQUIRE(Mutex.IsLocked());
-	});
+			AutoRTFM::Testing::Abort([&]
+			{
+				Mutex.Lock();
 
-	REQUIRE(!Mutex.IsLocked());
-}
+				AutoRTFM::Testing::Commit([&]
+				{
+					Mutex.Unlock();
+					Mutex.Lock();
+				});
 
-TEST_CASE("TransactionallySafeMutex.Abort(Lock, Commit(Unlock, Lock))")
-{
-	UE::FTransactionallySafeMutex Mutex;
+				RequireLocked(Mutex);
 
-	AutoRTFM::Testing::Abort([&]
-	{
-		Mutex.Lock();
+				AutoRTFM::AbortTransaction();
+			});
 
-		AutoRTFM::Testing::Commit([&]
+			RequireNotLocked(Mutex);
+		}
+	}
+
+	if constexpr (bSupportsUnbalancedLockingInTransaction)
+	{
+		SECTION("Contended Lock")
 		{
-			Mutex.Unlock();
-			Mutex.Lock();
-		});
+			// These subtests do not support transactional retries.
+			AUTORTFM_SCOPED_DISABLE_RETRY();
 
-		REQUIRE(Mutex.IsLocked());
+			SECTION("Contender Parks Transaction")
+			{
+				MutexType Mutex;
+				std::atomic_uint Orderer = 0;
 
-		AutoRTFM::AbortTransaction();
-	});
+				std::thread Contender([&] AUTORTFM_OPEN
+				{
+					REQUIRE(0 == Orderer);
+					Mutex.Lock();
+					Orderer += 1; // unblock main thread
+					while (1 == Orderer) {} // wait on main thread
+					std::this_thread::sleep_for(std::chrono::milliseconds(100));
+					Mutex.Unlock();
+				});
 
-	REQUIRE(!Mutex.IsLocked());
-}
+				while (0 == Orderer) {} // wait on contender
 
-TEST_CASE("TransactionallySafeMutex.Contended Lock")
-{
-	struct RAIIeroo final
+				AutoRTFM::Testing::Commit([&]
+				{
+					if constexpr (bSupportsIsLocked)
+					{
+						REQUIRE(Mutex.IsLocked());
+					}
+					AutoRTFM::Open([&] { Orderer += 1; }); // unblock contender
+					Mutex.Lock();
+				});
+
+				if constexpr (bSupportsIsLocked)
+				{
+					REQUIRE(Mutex.IsLocked());
+				}
+				Mutex.Unlock();
+				Contender.join();
+			}
+
+			SECTION("Transaction Parks Contender")
+			{
+				MutexType Mutex;
+				std::atomic_uint Orderer = 0;
+
+				std::thread Contender([&] AUTORTFM_OPEN
+				{
+					while (0 == Orderer) {} // wait on main thread
+					Mutex.Lock();
+				});
+
+				AutoRTFM::Testing::Commit([&]
+				{
+					Mutex.Lock();
+					AutoRTFM::Open([&]
+					{
+						Orderer += 1; // unblock contender
+						std::this_thread::sleep_for(std::chrono::milliseconds(100));
+					});
+					Mutex.Unlock();
+				});
+
+				Contender.join();
+				Mutex.Unlock();
+			}
+		}
+	}
+
+	SECTION("In On-Commit")
 	{
-		RAIIeroo()
+		if constexpr (bSupportsUnbalancedLockingInTransaction)
 		{
-			State = AutoRTFM::ForTheRuntime::GetRetryTransaction();
-			AutoRTFM::ForTheRuntime::SetRetryTransaction(AutoRTFM::ForTheRuntime::EAutoRTFMRetryTransactionState::NoRetry);
+			SECTION("Transact(Lock, Unlock-On-Commit)")
+			{
+				MutexType Mutex;
+				AutoRTFM::Testing::Commit([&]
+				{
+					Mutex.Lock();
+
+					AutoRTFM::OnCommit([&]
+					{
+						RequireLocked(Mutex);
+						Mutex.Unlock();
+						RequireNotLocked(Mutex);
+					});
+				});
+			}
 		}
 
-		~RAIIeroo()
+		if constexpr (bSupportsUnbalancedLockingInTransaction)
 		{
-			AutoRTFM::ForTheRuntime::SetRetryTransaction(State);
+			SECTION("Transact(Unlock-On-Commit, Lock)")
+			{
+				MutexType Mutex;
+				AutoRTFM::Testing::Commit([&]
+				{
+					AutoRTFM::OnCommit([&]
+					{
+						RequireLocked(Mutex);
+						Mutex.Unlock();
+						RequireLocked(Mutex);
+					});
+						
+					Mutex.Lock();
+				});
+			}
 		}
 
-	private:
-		AutoRTFM::ForTheRuntime::EAutoRTFMRetryTransactionState State;
-	} _;
+		if constexpr (bSupportsUnbalancedLockingInTransaction && bSupportsTryLockInTransaction)
+		{
+			SECTION("Transact(TryLock, Unlock-On-Commit)")
+			{
+				MutexType Mutex;
+				AutoRTFM::Testing::Commit([&]
+				{
+					REQUIRE(Mutex.TryLock());
+					REQUIRE(!Mutex.TryLock());
+
+					AutoRTFM::OnCommit([&]
+					{
+						RequireLocked(Mutex);
+						Mutex.Unlock();
+						RequireNotLocked(Mutex);
+					});
+				});
+			}
+		}
 
-	SECTION("Contender Parks Transaction")
-	{
-		UE::FTransactionallySafeMutex Mutex;
-		std::atomic_uint Orderer = 0;
+		if constexpr (bSupportsUnbalancedLockingInTransaction && bSupportsTryLockInTransaction)
+		{
+			SECTION("Transact(Unlock-On-Commit, TryLock)")
+			{
+				MutexType Mutex;
+				AutoRTFM::Testing::Commit([&]
+				{
+					AutoRTFM::OnCommit([&]
+					{
+						RequireLocked(Mutex);
+						Mutex.Unlock();
+						RequireLocked(Mutex);
+					});
+						
+					REQUIRE(Mutex.TryLock());
+				});
+			}
+		}
 
-		std::thread Contender([&] AUTORTFM_OPEN
+		if constexpr (bSupportsUnbalancedLockingInTransaction)
 		{
-			REQUIRE(0 == Orderer);
-			Mutex.Lock();
-			Orderer += 1; // unblock main thread
-			while (1 == Orderer) {} // wait on main thread
-			std::this_thread::sleep_for(std::chrono::milliseconds(100));
-			Mutex.Unlock();
-		});
+			SECTION("Lock, Transact(Lock-On-Commit, Unlock)")
+			{
+				MutexType Mutex;
+				Mutex.Lock();
 
-		while (0 == Orderer) {} // wait on contender
+				AutoRTFM::Testing::Commit([&]
+				{
+					AutoRTFM::OnCommit([&]
+					{
+						RequireLocked(Mutex);
+						Mutex.Lock();
+						RequireLocked(Mutex);
+					});
+						
+					Mutex.Unlock();
+				});
+			}
+		}
 
-		AutoRTFM::Testing::Commit([&]
+		if constexpr (bSupportsUnbalancedLockingInTransaction)
 		{
-			REQUIRE(Mutex.IsLocked());
-			AutoRTFM::Open([&] { Orderer += 1; }); // unblock contender
-			Mutex.Lock();
-		});
-
-		REQUIRE(Mutex.IsLocked());
-		Mutex.Unlock();
-		Contender.join();
-	}
+			SECTION("Lock, Transact(Unlock, Lock-On-Commit)")
+			{
+				MutexType Mutex;
+				Mutex.Lock();
 
-	SECTION("Transaction Parks Contender")
-	{
-		UE::FTransactionallySafeMutex Mutex;
-		std::atomic_uint Orderer = 0;
+				AutoRTFM::Testing::Commit([&]
+				{
+					Mutex.Unlock();
+
+					AutoRTFM::OnCommit([&]
+					{
+						RequireNotLocked(Mutex);
+						Mutex.Lock();
+						RequireLocked(Mutex);
+					});
+				});
+			}
+		}
 
-		std::thread Contender([&] AUTORTFM_OPEN
+		if constexpr (bSupportsUnbalancedLockingInTransaction && bSupportsTryLockInTransaction)
 		{
-			while (0 == Orderer) {} // wait on main thread
-			Mutex.Lock();
-		});
+			SECTION("Lock, Transact(Unlock, TryLock-On-Commit)")
+			{
+				MutexType Mutex;
+				Mutex.Lock();
 
-		AutoRTFM::Testing::Commit([&]
+				AutoRTFM::Testing::Commit([&]
+				{
+					Mutex.Unlock();
+
+					AutoRTFM::OnCommit([&]
+					{
+						REQUIRE(Mutex.TryLock());
+					});
+				});
+			}
+		}
+
+		if constexpr (bSupportsUnbalancedLockingInTransaction && bSupportsTryLockInTransaction)
 		{
-			Mutex.Lock();
-			AutoRTFM::Open([&]
+			SECTION("Lock, Transact(TryLock-On-Commit, Unlock)")
 			{
-				Orderer += 1; // unblock contender
-				std::this_thread::sleep_for(std::chrono::milliseconds(100));
-			});
-			Mutex.Unlock();
-		});
+				MutexType Mutex;
+				Mutex.Lock();
 
-		Contender.join();
-		Mutex.Unlock();
-	}
-}
+				AutoRTFM::Testing::Commit([&]
+				{
+					AutoRTFM::OnCommit([&]
+					{
+						REQUIRE(Mutex.TryLock());
+					});
 
-TEST_CASE("TransactionallySafeMutex.In On-Commit")
-{
-	UE::FTransactionallySafeMutex Mutex;
+					Mutex.Unlock();
+				});
+			}
+		}
+	}
 
-	AutoRTFM::Testing::Commit([&]
+	SECTION("In On-Abort")
 	{
-		Mutex.Lock();
-
-		AutoRTFM::OnCommit([&]
+		SECTION("Transact(Check-On-Abort, Lock, Check-On-Abort, Abort)")
 		{
-			REQUIRE(Mutex.IsLocked());
-			Mutex.Unlock();
-			REQUIRE(!Mutex.IsLocked());
-		});
-	});
+			MutexType Mutex;
+			AutoRTFM::Testing::Abort([&]
+			{
+				AutoRTFM::OnAbort([&]
+				{
+					RequireNotLocked(Mutex);
+				});
 
-	AutoRTFM::Testing::Commit([&]
-	{
-		AutoRTFM::OnCommit([&]
-		{
-			REQUIRE(Mutex.IsLocked());
-			Mutex.Unlock();
-			REQUIRE(Mutex.IsLocked());
-		});
-			
-		Mutex.Lock();
-	});
+				Mutex.Lock();
 
-	AutoRTFM::Testing::Commit([&]
-	{
-		REQUIRE(Mutex.TryLock());
-		REQUIRE(!Mutex.TryLock());
+				AutoRTFM::OnAbort([&]
+				{
+					RequireLocked(Mutex);
+				});
+
+				AutoRTFM::AbortTransaction();
+			});
+		}
 
-		AutoRTFM::OnCommit([&]
+		if constexpr (bSupportsTryLockInTransaction)
 		{
-			REQUIRE(Mutex.IsLocked());
-			Mutex.Unlock();
-			REQUIRE(!Mutex.IsLocked());
-		});
-	});
+			SECTION("Transact(Check-On-Abort, TryLock, Check-On-Abort, Abort)")
+			{
+				MutexType Mutex;
+				AutoRTFM::Testing::Abort([&]
+				{
+					AutoRTFM::OnAbort([&]
+					{
+						RequireNotLocked(Mutex);
+					});
+						
+					REQUIRE(Mutex.TryLock());
+
+					AutoRTFM::OnAbort([&]
+					{
+						RequireLocked(Mutex);
+					});
+
+					AutoRTFM::AbortTransaction();
+				});
+			}
+		}
 
-	AutoRTFM::Testing::Commit([&]
-	{
-		AutoRTFM::OnCommit([&]
+		if constexpr (bSupportsUnbalancedLockingInTransaction)
 		{
-			REQUIRE(Mutex.IsLocked());
-			Mutex.Unlock();
-			REQUIRE(Mutex.IsLocked());
-		});
-			
-		REQUIRE(Mutex.TryLock());
-	});
+			SECTION("Lock, Transact(Check-On-Abort, Unlock, Abort)")
+			{
+				MutexType Mutex;
+				Mutex.Lock();
 
-	Mutex.Lock();
-	REQUIRE(Mutex.IsLocked());
+				AutoRTFM::Testing::Abort([&]
+				{
+					AutoRTFM::OnAbort([&]
+					{
+						RequireLocked(Mutex);
+					});
+						
+					Mutex.Unlock();
+
+					AutoRTFM::AbortTransaction();
+				});
+			}
+		}
 
-	AutoRTFM::Testing::Commit([&]
-	{
-		AutoRTFM::OnCommit([&]
+		if constexpr (bSupportsUnbalancedLockingInTransaction)
 		{
-			REQUIRE(Mutex.IsLocked());
-			Mutex.Lock();
-			REQUIRE(Mutex.IsLocked());
-		});
-			
-		Mutex.Unlock();
-	});
+			SECTION("Lock, Transact(Unlock, Check-On-Abort, Abort)")
+			{
+				MutexType Mutex;
+				Mutex.Lock();
 
-	REQUIRE(Mutex.IsLocked());
+				AutoRTFM::Testing::Abort([&]
+				{
+					Mutex.Unlock();
 
-	AutoRTFM::Testing::Commit([&]
-	{
-		Mutex.Unlock();
+					AutoRTFM::OnAbort([&]
+					{
+						RequireLocked(Mutex);
+					});
 
-		AutoRTFM::OnCommit([&]
+					AutoRTFM::AbortTransaction();
+				});
+			}
+		}
+
+		if constexpr (bSupportsUnbalancedLockingInTransaction && bSupportsTryLockInTransaction)
 		{
-			REQUIRE(!Mutex.IsLocked());
-			Mutex.Lock();
-			REQUIRE(Mutex.IsLocked());
-		});
-	});
+			SECTION("Lock, Transact(Unlock, TryLock-On-Abort, Abort)")
+			{
+				MutexType Mutex;
+				Mutex.Lock();
 
-	REQUIRE(Mutex.IsLocked());
-		
-	AutoRTFM::Testing::Commit([&]
-	{
-		Mutex.Unlock();
+				AutoRTFM::Testing::Abort([&]
+				{
+					Mutex.Unlock();
 
-		AutoRTFM::OnCommit([&]
-		{
-			REQUIRE(Mutex.TryLock());
-		});
-	});
+					AutoRTFM::OnAbort([&]
+					{
+						REQUIRE(!Mutex.TryLock());
+					});
 
-	REQUIRE(Mutex.IsLocked());
-		
-	AutoRTFM::Testing::Commit([&]
-	{
-		AutoRTFM::OnCommit([&]
+					AutoRTFM::AbortTransaction();
+				});
+			}
+		}
+
+		if constexpr (bSupportsUnbalancedLockingInTransaction && bSupportsTryLockInTransaction)
 		{
-			REQUIRE(Mutex.TryLock());
-		});
+			SECTION("Lock, Transact(TryLock-On-Abort, Unlock, Abort)")
+			{
+				MutexType Mutex;
+				Mutex.Lock();
 
-		Mutex.Unlock();
-	});
+				AutoRTFM::Testing::Abort([&]
+				{
+					AutoRTFM::OnAbort([&]
+					{
+						REQUIRE(!Mutex.TryLock());
+					});
 
-	REQUIRE(Mutex.IsLocked());
-}
+					Mutex.Unlock();
 
-TEST_CASE("TransactionallySafeMutex.In On-Abort")
-{
-	UE::FTransactionallySafeMutex Mutex;
+					AutoRTFM::AbortTransaction();
+				});
+
+				RequireLocked(Mutex);
+			}
+		}
+	}
 
-	AutoRTFM::Testing::Abort([&]
+	SECTION("Locked Mutex In Destructed Object")
 	{
-		AutoRTFM::OnAbort([&]
+		struct AUTORTFM_ENABLE MyStruct final
 		{
-			REQUIRE(!Mutex.IsLocked());
-		});
-
-		Mutex.Lock();
+			MutexType Mutex;
 
-		AutoRTFM::OnAbort([&]
-		{
-			REQUIRE(Mutex.IsLocked());
-		});
+			~MyStruct()
+			{
+				memset(this, 0, sizeof(MyStruct));
+			}
+		};
 
-		AutoRTFM::AbortTransaction();
-	});
+		std::unique_ptr<MyStruct> Mine(new MyStruct);
 
-	AutoRTFM::Testing::Abort([&]
-	{
-		AutoRTFM::OnAbort([&]
+		AutoRTFM::Testing::Commit([&]
 		{
-			REQUIRE(!Mutex.IsLocked());
+			Mine->Mutex.Lock();
+			Mine.reset();
+			REQUIRE(!Mine);
 		});
-			
-		REQUIRE(Mutex.TryLock());
+	}
 
-		AutoRTFM::OnAbort([&]
+	if constexpr (bSupportsMixedOpenCloseLocking)
+	{
+		SECTION("Closed Lock Then Open Unlock")
 		{
-			REQUIRE(Mutex.IsLocked());
-		});
+			MutexType Mutex;
+			AutoRTFM::Testing::Commit([&]
+			{
+				Mutex.Lock();
+				RequireLocked(Mutex);
 
-		AutoRTFM::AbortTransaction();
-	});
+				AutoRTFM::Open([&]
+				{
+					Mutex.Unlock();
+				});
+			});
+			RequireNotLocked(Mutex);
+		}
+	}
 
-	AutoRTFM::Testing::Abort([&]
+	if constexpr (bSupportsMixedOpenCloseLocking && bSupportsTryLockInTransaction)
 	{
-		AutoRTFM::OnAbort([&]
+		SECTION("Closed TryLock Then Open Unlock")
 		{
-			REQUIRE(!Mutex.IsLocked());
-		});
-			
-		REQUIRE(Mutex.TryLock());
-
-		AutoRTFM::AbortTransaction();
-	});
+			MutexType Mutex;
+			AutoRTFM::Testing::Commit([&]
+			{
+				REQUIRE(Mutex.TryLock());
+				RequireLocked(Mutex);
 
-	Mutex.Lock();
-	REQUIRE(Mutex.IsLocked());
+				AutoRTFM::Open([&]
+				{
+					Mutex.Unlock();
+				});
+			});
+			RequireNotLocked(Mutex);
+		}
+	}
 
-	AutoRTFM::Testing::Abort([&]
+	if constexpr (bSupportsMixedOpenCloseLocking)
 	{
-		AutoRTFM::OnAbort([&]
+		SECTION("Open Lock Then Closed Unlock")
 		{
-			REQUIRE(Mutex.IsLocked());
-		});
-			
-		Mutex.Unlock();
-
-		AutoRTFM::AbortTransaction();
-	});
+			MutexType Mutex;
+			AutoRTFM::Testing::Commit([&]
+			{
+				AutoRTFM::Open([&]
+				{
+					Mutex.Lock();
+				});
+				RequireLocked(Mutex);
 
-	REQUIRE(Mutex.IsLocked());
+				Mutex.Unlock();
+			});
+			RequireNotLocked(Mutex);
+		}
+	}
 
-	AutoRTFM::Testing::Abort([&]
+	if constexpr (bSupportsMixedOpenCloseLocking && bSupportsTryLockInTransaction)
 	{
-		Mutex.Unlock();
-
-		AutoRTFM::OnAbort([&]
+		SECTION("Open TryLock Then Closed Unlock")
 		{
-			REQUIRE(Mutex.IsLocked());
-		});
+			MutexType Mutex;
+			AutoRTFM::Testing::Commit([&]
+			{
+				AutoRTFM::Open([&]
+				{
+					REQUIRE(Mutex.TryLock());
+				});
+				RequireLocked(Mutex);
 
-		AutoRTFM::AbortTransaction();
-	});
+				Mutex.Unlock();
+			});
+			RequireNotLocked(Mutex);
+		}
+	}
 
-	REQUIRE(Mutex.IsLocked());
-		
-	AutoRTFM::Testing::Abort([&]
+	if constexpr (bIsRecursive)
 	{
-		Mutex.Unlock();
-
-		AutoRTFM::OnAbort([&]
+		SECTION("Transact(Lock(Lock()))")
 		{
-			REQUIRE(!Mutex.TryLock());
-		});
-
-		AutoRTFM::AbortTransaction();
-	});
+			MutexType Mutex;
+			AutoRTFM::Testing::Commit([&]
+			{
+				RequireNotLocked(Mutex);
+				{
+					UE::TScopeLock LockA(Mutex);
+					RequireLocked(Mutex);
+					{
+						UE::TScopeLock LockB(Mutex);
+						RequireLocked(Mutex);
+					}
+					RequireLocked(Mutex);
+				}
+			});
+			RequireNotLocked(Mutex);
+		}
+	}
 
-	REQUIRE(Mutex.IsLocked());
-		
-	AutoRTFM::Testing::Abort([&]
+	if constexpr (bIsRecursive)
 	{
-		AutoRTFM::OnAbort([&]
+		SECTION("Lock(Lock())")
 		{
-			REQUIRE(!Mutex.TryLock());
-		});
-
-		Mutex.Unlock();
-
-		AutoRTFM::AbortTransaction();
-	});
-
-	REQUIRE(Mutex.IsLocked());
-}
+			MutexType Mutex;
+			RequireNotLocked(Mutex);
+			{
+				UE::TScopeLock LockA(Mutex);
+				RequireLocked(Mutex);
+				{
+					UE::TScopeLock LockB(Mutex);
+					RequireLocked(Mutex);
+				}
+			}
+			RequireNotLocked(Mutex);
+		}
+	}
 
-TEST_CASE("TransactionallySafeMutex.Locked Mutex In Destructed Object")
-{
-	struct AUTORTFM_ENABLE MyStruct final
+	if constexpr (bIsRecursive)
 	{
-		UE::FTransactionallySafeMutex Mutex;
-
-		~MyStruct()
+		SECTION("Transact(Lock(Lock(Abort)))")
 		{
-			memset(this, 0, sizeof(MyStruct));
+			MutexType Mutex;
+			AutoRTFM::Testing::Abort([&]
+			{
+				RequireNotLocked(Mutex);
+				{
+					UE::TScopeLock LockA(Mutex);
+					RequireLocked(Mutex);
+					{
+						UE::TScopeLock LockB(Mutex);
+						RequireLocked(Mutex);
+						AutoRTFM::AbortTransaction();
+					}
+				}
+			});
+			RequireNotLocked(Mutex);
 		}
-	};
-
-	std::unique_ptr<MyStruct> Mine(new MyStruct);
+	}
 
-	AutoRTFM::Testing::Commit([&]
+	if constexpr (bIsRecursive)
 	{
-		Mine->Mutex.Lock();
-		Mine.reset();
-		REQUIRE(!Mine);
-	});
-}
+		SECTION("Transact(Lock(Lock(), Abort))")
+		{
+			MutexType Mutex;
+			AutoRTFM::Testing::Abort([&]
+			{
+				RequireNotLocked(Mutex);
+				{
+					UE::TScopeLock LockA(Mutex);
+					RequireLocked(Mutex);
+					{
+						UE::TScopeLock LockB(Mutex);
+						RequireLocked(Mutex);
+					}
+					AutoRTFM::AbortTransaction();
+				}
+			});
+			RequireNotLocked(Mutex);
+		}
+	}
 
-TEST_CASE("TransactionallySafeMutex.Closed Lock Then Open Unlock")
-{
-	UE::FTransactionallySafeMutex Mutex;
-	AutoRTFM::Testing::Commit([&]
+	if constexpr (bIsRecursive)
 	{
-		Mutex.Lock();
-		REQUIRE(Mutex.IsLocked());
-
-		AutoRTFM::Open([&]
+		SECTION("Transact(Lock(Transact(Lock(), Abort)))")
 		{
-			Mutex.Unlock();
-		});
-	});
-	REQUIRE(!Mutex.IsLocked());
-}
+			MutexType Mutex;
+			AutoRTFM::Testing::Commit([&]
+			{
+				RequireNotLocked(Mutex);
+				UE::TScopeLock LockA(Mutex);
+				RequireLocked(Mutex);
+				AutoRTFM::Testing::Abort([&]
+				{
+					UE::TScopeLock LockB(Mutex);
+					RequireLocked(Mutex);
+					AutoRTFM::AbortTransaction();
+				});
+				RequireLocked(Mutex);
+			});
+			RequireNotLocked(Mutex);
+		}
+	}
 
-TEST_CASE("TransactionallySafeMutex.Closed TryLock Then Open Unlock")
-{
-	UE::FTransactionallySafeMutex Mutex;
-	AutoRTFM::Testing::Commit([&]
+	if constexpr (bIsShared && bIsRecursive)
 	{
-		REQUIRE(Mutex.TryLock());
-		REQUIRE(Mutex.IsLocked());
-
-		AutoRTFM::Open([&]
+		SECTION("SharedLock(SharedLock())")
 		{
-			Mutex.Unlock();
-		});
-	});
-	REQUIRE(!Mutex.IsLocked());
-}
+			MutexType Mutex;
+			RequireNotLocked(Mutex);
+			RequireNotSharedLocked(Mutex);
+			{
+				UE::TSharedLock LockA(Mutex);
+				RequireNotLocked(Mutex);
+				RequireLockShared(Mutex);
+				{
+					UE::TSharedLock LockB(Mutex);
+					RequireNotLocked(Mutex);
+					RequireLockShared(Mutex);
+				}
+				RequireNotLocked(Mutex);
+				RequireLockShared(Mutex);
+			}
+			RequireNotLocked(Mutex);
+			RequireNotSharedLocked(Mutex);
+		}
+	}
 
-TEST_CASE("TransactionallySafeMutex.Open Lock Then Closed Unlock")
-{
-	UE::FTransactionallySafeMutex Mutex;
-	AutoRTFM::Testing::Commit([&]
+	if constexpr (bIsShared && bIsRecursive)
 	{
-		AutoRTFM::Open([&]
+		SECTION("Transact(SharedLock(SharedLock()))")
 		{
-			Mutex.Lock();
-		});
-		REQUIRE(Mutex.IsLocked());
+			MutexType Mutex;
+			AutoRTFM::Testing::Commit([&]
+			{
+				RequireNotSharedLocked(Mutex);
+				{
+					UE::TSharedLock LockA(Mutex);
+					RequireLocked(Mutex);
+					RequireNotSharedLocked(Mutex);
+					{
+						UE::TSharedLock LockB(Mutex);
+						RequireLocked(Mutex);
+						RequireNotSharedLocked(Mutex);
+					}
+					RequireLocked(Mutex);
+					RequireNotSharedLocked(Mutex);
+				}
+			});
+			RequireNotLocked(Mutex);
+			RequireNotSharedLocked(Mutex);
+		}
+	}
 
-		Mutex.Unlock();
-	});
-	REQUIRE(!Mutex.IsLocked());
-}
+	if constexpr (bIsShared && bIsRecursive)
+	{
+		SECTION("Transact(SharedLock(SharedLock(Abort)))")
+		{
+			MutexType Mutex;
+			AutoRTFM::Testing::Abort([&]
+			{
+				RequireNotSharedLocked(Mutex);
+				{
+					UE::TSharedLock LockA(Mutex);
+					RequireLocked(Mutex);
+					RequireNotSharedLocked(Mutex);
+					{
+						UE::TSharedLock LockB(Mutex);
+						RequireLocked(Mutex);
+						RequireNotSharedLocked(Mutex);
+						AutoRTFM::AbortTransaction();
+					}
+				}
+			});
+			RequireNotLocked(Mutex);
+			RequireNotSharedLocked(Mutex);
+		}
+	}
 
-TEST_CASE("TransactionallySafeMutex.Open TryLock Then Closed Unlock")
-{
-	UE::FTransactionallySafeMutex Mutex;
-	AutoRTFM::Testing::Commit([&]
+	if constexpr (bIsShared && bIsRecursive)
 	{
-		AutoRTFM::Open([&]
+		SECTION("Transact(SharedLock(SharedLock(), Abort))")
 		{
-			REQUIRE(Mutex.TryLock());
-		});
-		REQUIRE(Mutex.IsLocked());
+			MutexType Mutex;
+			AutoRTFM::Testing::Abort([&]
+			{
+				RequireNotSharedLocked(Mutex);
+				{
+					UE::TSharedLock LockA(Mutex);
+					RequireLocked(Mutex);
+					RequireNotSharedLocked(Mutex);
+					{
+						UE::TSharedLock LockB(Mutex);
+						RequireLocked(Mutex);
+						RequireNotSharedLocked(Mutex);
+					}
+					AutoRTFM::AbortTransaction();
+				}
+			});
+			RequireNotLocked(Mutex);
+			RequireNotSharedLocked(Mutex);
+		}
+	}
 
-		Mutex.Unlock();
-	});
-	REQUIRE(!Mutex.IsLocked());
+	if constexpr (bIsShared && bIsRecursive)
+	{
+		SECTION("Transact(SharedLock(Transact(SharedLock(), Abort)))")
+		{
+			MutexType Mutex;
+			AutoRTFM::Testing::Commit([&]
+			{
+				RequireNotSharedLocked(Mutex);
+				UE::TSharedLock LockA(Mutex);
+				RequireLocked(Mutex);
+				RequireNotSharedLocked(Mutex);
+				AutoRTFM::Testing::Abort([&]
+				{
+					UE::TSharedLock LockB(Mutex);
+					RequireLocked(Mutex);
+					RequireNotSharedLocked(Mutex);
+					AutoRTFM::AbortTransaction();
+				});
+				RequireLocked(Mutex);
+				RequireNotSharedLocked(Mutex);
+			});
+			RequireNotLocked(Mutex);
+			RequireNotSharedLocked(Mutex);
+		}
+	}
 }
diff --git a/Engine/Source/Runtime/Core/Public/Async/TransactionallySafeSharedRecursiveMutex.h b/Engine/Source/Runtime/Core/Public/Async/TransactionallySafeSharedRecursiveMutex.h
new file mode 100644
--- /dev/null
+++ b/Engine/Source/Runtime/Core/Public/Async/TransactionallySafeSharedRecursiveMutex.h
@@ -0,0 +1,221 @@
+// Copyright Epic Games, Inc. All Rights Reserved.
+
+#pragma once
+
+#include "Async/SharedRecursiveMutex.h"
+#include "AutoRTFM.h"
+#include "Templates/SharedPointerFwd.h"
+
+#if UE_AUTORTFM
+#include "AutoRTFM/OpenWrapper.h"
+#include "Templates/SharedPointer.h"
+#endif
+
+namespace UE
+{
+
+#if UE_AUTORTFM
+
+// A transactionally-safe shared lock that supports recursive locking and works in the following novel ways:
+// - In the open (non-transactional):
+//   - Take the lock like before. Simple!
+//   - Free the lock like before too.
+// - In the closed (transactional):
+//   - During locking we query `TransactionalLockCount`:
+//	   - 0 means we haven't taken the lock within our transaction nest and need to acquire the lock.
+//     - Otherwise we already have the lock (and are preventing non-transactional code seeing any
+//       modifications we've made while holding the lock), so just bump `TransactionalLockCount`.
+//     - We also register an on-abort handler to release the lock should we abort (but we need to
+//       query `TransactionalLockCount` even there because we could be aborting an inner transaction
+//       and the parent transaction still wants to have the lock held!).
+//   - During unlocking we defer doing the unlock until the transaction commits.
+//
+// Thus with this approach we will hold this lock for the *entirety* of the transactional nest should
+// we take the lock during the transaction, thus preventing non-transactional code from seeing any
+// modifications we should make.
+//
+// If we are within a transaction, we pessimise our shared lock to a full lock. Note: that it should
+// potentially be possible to have shared locks work correctly, but serious care will have to be taken to
+// ensure that we don't have:
+//   Open Thread     Closed Thread
+//   -----------     SharedLock
+//   -----------     SharedUnlock
+//   Lock            -------------
+//   Unlock          -------------
+//   -----------     SharedLock      <- Invalid because the transaction can potentially observe side
+//                                      effects of the open-threads writes!
+struct FTransactionallySafeSharedRecursiveMutexDefinition
+{
+	FTransactionallySafeSharedRecursiveMutexDefinition() : State(MakeShared<FState>())
+	{
+	}
+
+	void LockShared(Core::Private::FSharedRecursiveMutexLink& Link)
+	{
+		if (AutoRTFM::IsTransactional() || AutoRTFM::IsCommittingOrAborting())
+		{
+			// Transactionally pessimise LockShared -> Lock.
+			Lock();
+		}
+		else
+		{
+			State->Mutex.LockShared(Link);
+			ensure(0 == State->TransactionalLockCount);
+		}
+	}
+
+	void UnlockShared(Core::Private::FSharedRecursiveMutexLink& Link)
+	{
+		if (AutoRTFM::IsTransactional() || AutoRTFM::IsCommittingOrAborting())
+		{
+			// Transactionally pessimise UnlockShared -> Unlock.
+			Unlock();
+		}
+		else
+		{
+			ensure(0 == State->TransactionalLockCount);
+			State->Mutex.UnlockShared(Link);
+		}
+	}
+
+	void Lock()
+	{
+		if (AutoRTFM::IsTransactional() || AutoRTFM::IsCommittingOrAborting())
+		{
+			UE_AUTORTFM_OPEN
+			{
+				// The transactional system which can increment TransactionalLockCount
+				// is always single-threaded, thus this is safe to check without atomicity.
+				if (0 == State->TransactionalLockCount)
+				{
+					State->Mutex.Lock();
+				}
+
+				State->TransactionalLockCount += 1;
+			};
+
+			// We explicitly copy the state here for the case that `this` was stack
+			// allocated and has already died before the on-abort is hit.
+			AutoRTFM::OnAbort([State = AutoRTFM::TOpenWrapper{this->State}]
+			{
+				ensure(0 != State.Object->TransactionalLockCount);
+
+				State.Object->TransactionalLockCount -= 1;
+
+				if (0 == State.Object->TransactionalLockCount)
+				{
+					State.Object->Mutex.Unlock();
+				}
+			});
+		}
+		else
+		{
+			State->Mutex.Lock();
+			ensure(0 == State->TransactionalLockCount);
+		}
+	}
+
+	void Unlock()
+	{
+		if (AutoRTFM::IsTransactional() || AutoRTFM::IsCommittingOrAborting())
+		{
+			// We explicitly copy the state here for the case that `this` was stack
+			// allocated and has already died before the on-commit is hit.
+			AutoRTFM::OnCommit([State = AutoRTFM::TOpenWrapper{this->State}]
+			{
+				ensure(0 != State.Object->TransactionalLockCount);
+
+				State.Object->TransactionalLockCount -= 1;
+
+				if (0 == State.Object->TransactionalLockCount)
+				{
+					State.Object->Mutex.Unlock();
+				}
+			});
+		}
+		else
+		{
+			ensure(0 == State->TransactionalLockCount);
+			State->Mutex.Unlock();
+		}
+	}
+
+	/** 
+	 * This method may give surprising results and should be used with caution!
+	 * 
+	 * - You cannot safely use this function to determine whether Lock() will block.
+	 *   You may be in race with another thread which is also trying to lock the mutex.
+	 * 
+	 * - If a Mutex is locked by the AutoRTFM thread, it will not fully release the mutex
+	 *   until the transaction fully succeeds or rolls back. Because of this, IsLocked may
+	 *   return true even after the mutex has been fully unlocked.
+	 */
+	AUTORTFM_OPEN [[nodiscard]] bool IsLocked() const
+	{
+		return State->Mutex.IsLocked();
+	}
+
+	/** 
+	 * This method may give surprising results and should be used with caution!
+	 * 
+	 * - You cannot safely use this function to determine whether LockShared() will block.
+	 *   You may be in race with another thread which is also trying to lock the mutex.
+	 * 
+	 * - If a Mutex is locked by the AutoRTFM thread, it will not fully release the mutex
+	 *   until the transaction fully succeeds or rolls back. Because of this, IsLockedShared may
+	 *   return true even after the mutex has been fully unlocked.
+	 */
+	AUTORTFM_OPEN [[nodiscard]] bool IsLockShared() const
+	{
+		return State->Mutex.IsLockShared();
+	}
+
+private:
+	UE_NONCOPYABLE(FTransactionallySafeSharedRecursiveMutexDefinition)
+
+	struct FState final
+	{
+		::UE::FSharedRecursiveMutex Mutex;
+		uint32 TransactionalLockCount = 0;
+
+		FState() = default;
+		
+		~FState()
+		{
+			ensure(0 == TransactionalLockCount);
+		}
+	};
+
+	TSharedPtr<FState> State;
+};
+
+using FTransactionallySafeSharedRecursiveMutex = ::UE::FTransactionallySafeSharedRecursiveMutexDefinition;
+
+template <>
+class TSharedLock<FTransactionallySafeSharedRecursiveMutex> final
+{
+public:
+	TSharedLock(const TSharedLock&) = delete;
+	TSharedLock& operator=(const TSharedLock&) = delete;
+
+	[[nodiscard]] inline explicit TSharedLock(FTransactionallySafeSharedRecursiveMutex& Lock)
+		: Mutex(Lock)
+	{
+		Mutex.LockShared(Link);
+	}
+
+	inline ~TSharedLock()
+	{
+		Mutex.UnlockShared(Link);
+	}
+
+private:
+	FTransactionallySafeSharedRecursiveMutex& Mutex;
+	Core::Private::FSharedRecursiveMutexLink Link;
+};
+
+#else
+using FTransactionallySafeSharedRecursiveMutex = ::UE::FSharedRecursiveMutex;
+#endif
+
+}  // namespace UE
