[AutoRTFM] Add an AutoRTFM hazard message for UE::FRecursiveMutex.

We will now prompt users to switch to an FTransactionallySafeRecursiveMutex instead.

#Synced-CL 48085702

diff --git a/Engine/Source/Programs/AutoRTFMTests/Private/Abort.cpp b/Engine/Source/Programs/AutoRTFMTests/Private/Abort.cpp
--- a/Engine/Source/Programs/AutoRTFMTests/Private/Abort.cpp
+++ b/Engine/Source/Programs/AutoRTFMTests/Private/Abort.cpp
@@ -5,6 +5,7 @@
 #include "AutoRTFMTestUtils.h"
 #include "AutoRTFM/AutoRTFMUE.h"
 #include "Async/Mutex.h"
+#include "Async/RecursiveMutex.h"
 #include "Async/SharedMutex.h"
 #include "Catch2Includes.h"
 #include "Misc/StringOutputDevice.h"
@@ -1517,6 +1518,42 @@ TEST_CASE("Abort.Language.ErrorHints.SharedMutexIsLockShared")
 	REQUIRE(WarningContext.HasWarningSubstring("A shared mutex was accessed while inside an AutoRTFM transaction."));
 }
 
+TEST_CASE("Abort.Language.ErrorHints.RecursiveMutexLock")
+{
+	FScopedErrorHintTest ScopedTest;
+    AutoRTFMTestUtils::FCaptureWarningContext WarningContext;
+
+	UE::FRecursiveMutex Mutex;
+
+	const AutoRTFM::ETransactionResult Result = AutoRTFM::Transact([&]
+	{
+		Mutex.Lock();
+	});
+
+	REQUIRE(AutoRTFM::ETransactionResult::AbortedByLanguage == Result);
+
+	// Did we print our message?
+	REQUIRE(WarningContext.HasWarningSubstring("A recursive mutex was accessed while inside an AutoRTFM transaction."));
+}
+
+TEST_CASE("Abort.Language.ErrorHints.RecursiveMutexIsLocked")
+{
+	FScopedErrorHintTest ScopedTest;
+    AutoRTFMTestUtils::FCaptureWarningContext WarningContext;
+
+	UE::FRecursiveMutex Mutex;
+
+	const AutoRTFM::ETransactionResult Result = AutoRTFM::Transact([&]
+	{
+		REQUIRE(!Mutex.IsLocked());
+	});
+
+	REQUIRE(AutoRTFM::ETransactionResult::AbortedByLanguage == Result);
+
+	// Did we print our message?
+	REQUIRE(WarningContext.HasWarningSubstring("A recursive mutex was accessed while inside an AutoRTFM transaction."));
+}
+
 #if 0
 // Testing an open free is flaky, since it leads to a double free--one inside the open block, and one during the
 // transactional abort caused by an AutoRTFM hazard. This can cause the unit test harness to fail on a subsequent test,
diff --git a/Engine/Source/Runtime/AutoRTFM/Private/ReportHazard.cpp b/Engine/Source/Runtime/AutoRTFM/Private/ReportHazard.cpp
--- a/Engine/Source/Runtime/AutoRTFM/Private/ReportHazard.cpp
+++ b/Engine/Source/Runtime/AutoRTFM/Private/ReportHazard.cpp
@@ -49,6 +49,16 @@ it does support transactionally-safe mutexes that can be used instead.
 
 Use the call stack below to find the FMutex accessed immediately before this error occurred, and
 replace it with FTransactionallySafeMutex.
+)";
+			break;
+
+		case EHazardType::RecursiveMutex:
+			Message = R"(
+A recursive mutex was accessed while inside an AutoRTFM transaction. AutoRTFM doesn't support native
+recursive mutexes, but it does support transactionally-safe recursive mutexes that can be used instead.
+
+Use the call stack below to find the FRecursiveMutex accessed immediately before this error occurred,
+and replace it with FTransactionallySafeRecursiveMutex.
 )";
 			break;
 
diff --git a/Engine/Source/Runtime/AutoRTFM/Public/AutoRTFM/ReportHazard.h b/Engine/Source/Runtime/AutoRTFM/Public/AutoRTFM/ReportHazard.h
--- a/Engine/Source/Runtime/AutoRTFM/Public/AutoRTFM/ReportHazard.h
+++ b/Engine/Source/Runtime/AutoRTFM/Public/AutoRTFM/ReportHazard.h
@@ -15,6 +15,7 @@ enum class EHazardType
 	CriticalSection,
 	RWLock,
 	Mutex,
+	RecursiveMutex,
 	SharedMutex,
 	OpenFree,
 };
@@ -31,10 +32,10 @@ AUTORTFM_OPEN UE_AUTORTFM_API void ReportAutoRTFMHazard(EHazardType HazardType);
  * Reports an AutoRTFM hazard if this block of code is reached from the closed.
  * `autortfm_is_closed` is optimized away at compile time, so this compiles away to nothing in an open function.
  */
-#define UE_AUTORTFM_REPORT_HAZARD_IF_CLOSED(InHazardType)                                                           \
-	(autortfm_is_closed()                                                                                       \
-		? ::AutoRTFM::ForTheRuntime::ReportAutoRTFMHazard(::AutoRTFM::ForTheRuntime::EHazardType::InHazardType) \
-		: (void)0)
+#define UE_AUTORTFM_REPORT_HAZARD_IF_CLOSED(InHazardType)                                                       \
+    (autortfm_is_closed()                                                                                       \
+        ? ::AutoRTFM::ForTheRuntime::ReportAutoRTFMHazard(::AutoRTFM::ForTheRuntime::EHazardType::InHazardType) \
+        : (void)0)
 
 #else
 
diff --git a/Engine/Source/Runtime/Core/Private/Async/RecursiveMutex.cpp b/Engine/Source/Runtime/Core/Private/Async/RecursiveMutex.cpp
--- a/Engine/Source/Runtime/Core/Private/Async/RecursiveMutex.cpp
+++ b/Engine/Source/Runtime/Core/Private/Async/RecursiveMutex.cpp
@@ -12,6 +12,8 @@ namespace UE
 
 bool FRecursiveMutex::TryLock()
 {
+	UE_AUTORTFM_REPORT_HAZARD_IF_CLOSED(RecursiveMutex);
+
 	const uint32 CurrentThreadId = FPlatformTLS::GetCurrentThreadId();
 	uint32 CurrentState = State.load(std::memory_order_relaxed);
 
@@ -39,6 +41,8 @@ bool FRecursiveMutex::TryLock()
 
 void FRecursiveMutex::Lock()
 {
+	UE_AUTORTFM_REPORT_HAZARD_IF_CLOSED(RecursiveMutex);
+
 	const uint32 CurrentThreadId = FPlatformTLS::GetCurrentThreadId();
 	uint32 CurrentState = State.load(std::memory_order_relaxed);
 
@@ -114,6 +118,8 @@ FORCENOINLINE void FRecursiveMutex::LockSlow(uint32 CurrentState, const uint32 C
 
 void FRecursiveMutex::Unlock()
 {
+	UE_AUTORTFM_REPORT_HAZARD_IF_CLOSED(RecursiveMutex);
+
 	uint32 CurrentState = State.load(std::memory_order_relaxed);
 	checkSlow(CurrentState & LockCountMask);
 	checkSlow(ThreadId.load(std::memory_order_relaxed) == FPlatformTLS::GetCurrentThreadId());
diff --git a/Engine/Source/Runtime/Core/Public/Async/RecursiveMutex.h b/Engine/Source/Runtime/Core/Public/Async/RecursiveMutex.h
--- a/Engine/Source/Runtime/Core/Public/Async/RecursiveMutex.h
+++ b/Engine/Source/Runtime/Core/Public/Async/RecursiveMutex.h
@@ -3,6 +3,7 @@
 #pragma once
 
 #include "CoreTypes.h"
+#include "AutoRTFM/ReportHazard.h"
 #include <atomic>
 
 #define UE_API CORE_API
@@ -25,6 +26,7 @@ public:
 
 	[[nodiscard]] inline bool IsLocked() const
 	{
+		UE_AUTORTFM_REPORT_HAZARD_IF_CLOSED(RecursiveMutex);
 		return !!(State.load(std::memory_order_relaxed) & LockCountMask);
 	}
 
