Override SerializeRecord in FCaptureWarningContext, otherwise FFeedbackContext swallows things in its own implementation, unlike the base implementation in FOutputDevice which forwards records to Serialize as strings.
Also clean up duplicate of FCaptureWarningContext.

#Synced-CL Unknown

diff --git a/Engine/Source/Programs/AutoRTFMTests/Private/AutoRTFMTestUtils.h b/Engine/Source/Programs/AutoRTFMTests/Private/AutoRTFMTestUtils.h
--- a/Engine/Source/Programs/AutoRTFMTests/Private/AutoRTFMTestUtils.h
+++ b/Engine/Source/Programs/AutoRTFMTests/Private/AutoRTFMTestUtils.h
@@ -65,6 +65,13 @@ public:
 		}
 	}
 
+	UE_AUTORTFM_ALWAYS_OPEN
+	void SerializeRecord(const UE::FLogRecord& Record) override
+	{
+		// Skip FFeedbackContext and allow the base impl to turn the record into a string so we can store it
+		FOutputDevice::SerializeRecord(Record);
+	}
+
 	UE_AUTORTFM_ALWAYS_OPEN
 	const std::vector<std::string>& GetWarnings() const
 	{
diff --git a/Engine/Source/Programs/AutoRTFMTests/Private/CString.cpp b/Engine/Source/Programs/AutoRTFMTests/Private/CString.cpp
--- a/Engine/Source/Programs/AutoRTFMTests/Private/CString.cpp
+++ b/Engine/Source/Programs/AutoRTFMTests/Private/CString.cpp
@@ -2,6 +2,7 @@
 
 #include "AutoRTFM.h"
 #include "AutoRTFM/Testing.h"
+#include "AutoRTFMTestUtils.h"
 #include "Catch2Includes.h"
 #include "CoreGlobals.h"
 #include "Logging/LogVerbosity.h"
@@ -43,7 +44,7 @@ namespace
 
 // The expected AutoRTFM warning message when attempting to printf() using a
 // format string containing a '%n' format specifier.
-static constexpr auto* kPercentNWarning = TEXT("AutoRTFM does not support format strings containing '%n'");
+static constexpr char kPercentNWarning[] = "AutoRTFM does not support format strings containing '%n'";
 
 // CheckWCharFunction calls Function with WChar outside a transaction, in an
 // aborted transaction and committed transaction. The value returned when called
@@ -206,51 +207,6 @@ void CheckStdToChars(ARGS ... Args)
 	}
 }
 
-// A helper class that for the lifetime of the object, intercepts and records
-// UE_LOG warnings.
-class FCaptureWarningContext : private FFeedbackContext
-{
-public:
-	FCaptureWarningContext() : OldContext(GWarn) { GWarn = this; }
-	~FCaptureWarningContext() { GWarn = OldContext; }
-
-	void Serialize(const TCHAR* V, ELogVerbosity::Type Verbosity, const FName& Category) override
-	{
-		if (Verbosity == ELogVerbosity::Warning)
-		{
-			Warnings.push_back(V);
-		}
-		else
-		{
-			OldContext->Serialize(V, Verbosity, Category);
-		}
-	}
-
-	void Serialize(const TCHAR* V, ELogVerbosity::Type Verbosity, const FName& Category, double Time) override
-	{
-		if (Verbosity == ELogVerbosity::Warning)
-		{
-			Warnings.push_back(V);
-		}
-		else
-		{
-			OldContext->Serialize(V, Verbosity, Category, Time);
-		}
-	}
-
-	const std::vector<FString>& GetWarnings() const
-	{
-		return Warnings;
-	}
-
-private:
-	FCaptureWarningContext(const FCaptureWarningContext&) = delete;
-	FCaptureWarningContext& operator = (const FCaptureWarningContext&) = delete;
-
-	FFeedbackContext* OldContext = nullptr;
-	std::vector<FString> Warnings;
-};
-
 }  // anonymous namespace
 
 TEST_CASE("CString.memcpy")
@@ -1438,7 +1394,7 @@ TEST_CASE("CString.snprintf")
 
 	SECTION("PercentN")
 	{
-		FCaptureWarningContext WarningContext;
+		AutoRTFMTestUtils::FCaptureWarningContext WarningContext;
 
 		int Count = 0;
 
@@ -1452,8 +1408,7 @@ TEST_CASE("CString.snprintf")
 		REQUIRE(AutoRTFM::ETransactionResult::AbortedByLanguage == Result);
 		REQUIRE(0 == Count);
 		REQUIRE("_____________________" == std::string_view(Buffer));
-		REQUIRE_THAT(WarningContext.GetWarnings(),
-			Catch::Matchers::VectorContains(FString(kPercentNWarning)));
+		REQUIRE(WarningContext.HasWarning(kPercentNWarning));
 	}
 
 	SECTION("PercentPercentN")
@@ -1606,7 +1561,7 @@ TEST_CASE("CString.swprintf")
 
 	SECTION("PercentN")
 	{
-		FCaptureWarningContext WarningContext;
+		AutoRTFMTestUtils::FCaptureWarningContext WarningContext;
 
 		int Count = 0;
 
@@ -1620,8 +1575,7 @@ TEST_CASE("CString.swprintf")
 		REQUIRE(AutoRTFM::ETransactionResult::AbortedByLanguage == Result);
 		REQUIRE(0 == Count);
 		REQUIRE(std::wstring_view(L"_____________________") == std::wstring_view(Buffer));
-		REQUIRE_THAT(WarningContext.GetWarnings(),
-			Catch::Matchers::VectorContains(FString(kPercentNWarning)));
+		REQUIRE(WarningContext.HasWarning(kPercentNWarning));
 	}
 
 	SECTION("PercentPercentN")
@@ -1694,7 +1648,7 @@ TEST_CASE("CString.printf")
 
 	SECTION("PercentN")
 	{
-		FCaptureWarningContext WarningContext;
+		AutoRTFMTestUtils::FCaptureWarningContext WarningContext;
 
 		int Count = 0;
 
@@ -1707,8 +1661,7 @@ TEST_CASE("CString.printf")
 
 		REQUIRE(AutoRTFM::ETransactionResult::AbortedByLanguage == Result);
 		REQUIRE(0 == Count);
-		REQUIRE_THAT(WarningContext.GetWarnings(),
-			Catch::Matchers::VectorContains(FString(kPercentNWarning)));
+		REQUIRE(WarningContext.HasWarning(kPercentNWarning));
 	}
 }
 
@@ -1749,7 +1702,7 @@ TEST_CASE("CString.wprintf")
 
 	SECTION("PercentN")
 	{
-		FCaptureWarningContext WarningContext;
+		AutoRTFMTestUtils::FCaptureWarningContext WarningContext;
 
 		int Count = 0;
 
@@ -1762,8 +1715,7 @@ TEST_CASE("CString.wprintf")
 
 		REQUIRE(AutoRTFM::ETransactionResult::AbortedByLanguage == Result);
 		REQUIRE(0 == Count);
-		REQUIRE_THAT(WarningContext.GetWarnings(),
-			Catch::Matchers::VectorContains(FString(kPercentNWarning)));
+		REQUIRE(WarningContext.HasWarning(kPercentNWarning));
 	}
 }
 
