Maintain structured logging when passing non-warning records to OldContext.

#Synced-CL Unknown

diff --git a/Engine/Source/Programs/AutoRTFMTests/Private/AutoRTFMTestUtils.h b/Engine/Source/Programs/AutoRTFMTests/Private/AutoRTFMTestUtils.h
--- a/Engine/Source/Programs/AutoRTFMTests/Private/AutoRTFMTestUtils.h
+++ b/Engine/Source/Programs/AutoRTFMTests/Private/AutoRTFMTestUtils.h
@@ -6,6 +6,7 @@
 #include "AutoRTFM/Defines.h"
 #include "Containers/StringConv.h"
 #include "HAL/Platform.h"
+#include "Logging/StructuredLog.h"
 #include "Misc/FeedbackContext.h"
 
 #include <algorithm>
@@ -68,8 +69,15 @@ public:
 	UE_AUTORTFM_ALWAYS_OPEN
 	void SerializeRecord(const UE::FLogRecord& Record) override
 	{
-		// Skip FFeedbackContext and allow the base impl to turn the record into a string so we can store it
-		FOutputDevice::SerializeRecord(Record);
+		if (Record.GetVerbosity() == ELogVerbosity::Warning)
+		{
+			// Skip FFeedbackContext and allow the base impl to turn the record into a string so we can store it
+			FOutputDevice::SerializeRecord(Record);
+		}
+		else
+		{
+			OldContext->SerializeRecord(Record);
+		}
 	}
 
 	UE_AUTORTFM_ALWAYS_OPEN
