Remove `Commit` and `autortfm_commit` from the main `AutoRTFM.h` header, and use `AutoRTFMTesting.h`'s `Testing::Commit` instead since this is a test-only feature.

#rnx
#Synced-CL 48550887

diff --git a/Engine/Source/Programs/AutoRTFMTests/Private/API.cpp b/Engine/Source/Programs/AutoRTFMTests/Private/API.cpp
--- a/Engine/Source/Programs/AutoRTFMTests/Private/API.cpp
+++ b/Engine/Source/Programs/AutoRTFMTests/Private/API.cpp
@@ -14,7 +14,7 @@ TEST_CASE("API.autortfm_is_transactional")
     bool bInTransaction = false;
     bool bInOpenNest = false;
 
-    AutoRTFM::Commit([&]
+    AutoRTFM::Testing::Commit([&]
     {
         bInTransaction = autortfm_is_transactional();
 
@@ -37,7 +37,7 @@ TEST_CASE("API.autortfm_is_closed")
     bool bInOpenNest = true;
     bool bInClosedNestInOpenNest = false;
 
-    AutoRTFM::Commit([&]
+    AutoRTFM::Testing::Commit([&]
     {
         bInTransaction = autortfm_is_closed();
 
@@ -63,7 +63,7 @@ TEST_CASE("API.autortfm_abort_transaction")
     bool bInNest = false;
     bool bAfterNest = false;
     
-    AutoRTFM::Commit([&]
+    AutoRTFM::Testing::Commit([&]
     {
         bBeforeNest = true;
 
@@ -129,7 +129,7 @@ TEST_CASE("API.autortfm_register_open_to_closed_functions")
 
     int I = -42;
 
-    AutoRTFM::Commit([&]
+    AutoRTFM::Testing::Commit([&]
     {
         I = NoAutoRTFM::DoSomethingC(I);
     });
@@ -144,7 +144,7 @@ TEST_CASE("API.autortfm_on_commit")
     bool bInnerTransactionWithAbort = false;
     bool bInnerOpenNest = false;
 
-    AutoRTFM::Commit([&]
+    AutoRTFM::Testing::Commit([&]
     {
         autortfm_on_commit([](void* const Arg)
         {
@@ -157,7 +157,7 @@ TEST_CASE("API.autortfm_on_commit")
             AutoRTFM::AbortTransaction();
         }
 
-        AutoRTFM::Commit([&]
+        AutoRTFM::Testing::Commit([&]
         {
             autortfm_on_commit([](void* const Arg)
             {
@@ -245,7 +245,7 @@ TEST_CASE("API.autortfm_on_abort")
             AutoRTFM::AbortTransaction();
         }
 
-        AutoRTFM::Commit([&]
+        AutoRTFM::Testing::Commit([&]
         {
             autortfm_on_abort([](void* const Arg)
             {
@@ -304,7 +304,7 @@ TEST_CASE("API.autortfm_did_allocate")
     unsigned BumpAllocator[Size];
     unsigned NextBump = 0;
 
-    AutoRTFM::Commit([&]
+    AutoRTFM::Testing::Commit([&]
     {
 		// If we are retrying transactions, need to reset the test state.
 		AutoRTFM::OnAbort([&]
@@ -673,7 +673,7 @@ TEST_CASE("API.Commit")
 {
     int Answer = 6 * 9;
 
-    AutoRTFM::Commit([&]
+    AutoRTFM::Testing::Commit([&]
     {
         Answer = 42;
     });
@@ -687,7 +687,7 @@ TEST_CASE("API.Abort")
     bool bInNest = false;
     bool bAfterNest = false;
 
-    AutoRTFM::Commit([&]
+    AutoRTFM::Testing::Commit([&]
     {
         bBeforeNest = true;
 
@@ -819,7 +819,7 @@ TEST_CASE("API.OnCommit")
             AutoRTFM::AbortTransaction();
         }
 
-        AutoRTFM::Commit([&]
+        AutoRTFM::Testing::Commit([&]
         {
             AutoRTFM::OnCommit([&]
             {
@@ -1084,7 +1084,7 @@ TEST_CASE("API.OnAbort")
 			AutoRTFM::AbortTransaction();
         }
 
-        AutoRTFM::Commit([&]
+        AutoRTFM::Testing::Commit([&]
         {
             AutoRTFM::OnAbort([&]
             {
@@ -1172,7 +1172,7 @@ TEST_CASE("API.OnPreAbort")
 			AutoRTFM::AbortTransaction();
         }
 
-        AutoRTFM::Commit([&]
+        AutoRTFM::Testing::Commit([&]
         {
             AutoRTFM::OnPreAbort([&]
             {
@@ -1335,7 +1335,7 @@ TEST_CASE("API.DidAllocate")
     unsigned BumpAllocator[Size];
     unsigned NextBump = 0;
 
-    AutoRTFM::Commit([&]
+    AutoRTFM::Testing::Commit([&]
     {
 		// If we are retrying transactions, need to reset the test state.
 		AutoRTFM::OnAbort([&]
@@ -1374,7 +1374,7 @@ TEST_CASE("API.IsOnCurrentTransactionStack")
 		delete OnHeapNotInTransaction;
 	}
 
-	AutoRTFM::Commit([&]
+	AutoRTFM::Testing::Commit([&]
 	{
 		int OnStackInTransaction = 3;
 		REQUIRE(AutoRTFM::IsOnCurrentTransactionStack(&OnStackInTransaction));
@@ -1383,7 +1383,7 @@ TEST_CASE("API.IsOnCurrentTransactionStack")
 		REQUIRE(!AutoRTFM::IsOnCurrentTransactionStack(OnHeapInTransaction));
 		delete OnHeapInTransaction;
 
-		AutoRTFM::Commit([&]
+		AutoRTFM::Testing::Commit([&]
 		{
 			// `OnStackInTransaction` is no longer in the innermost scope.
 			REQUIRE(!AutoRTFM::IsOnCurrentTransactionStack(&OnStackInTransaction));
diff --git a/Engine/Source/Programs/AutoRTFMTests/Private/Abort.cpp b/Engine/Source/Programs/AutoRTFMTests/Private/Abort.cpp
--- a/Engine/Source/Programs/AutoRTFMTests/Private/Abort.cpp
+++ b/Engine/Source/Programs/AutoRTFMTests/Private/Abort.cpp
@@ -62,7 +62,7 @@ TEST_CASE("Abort.NestedAbortOrder")
 {
 	unsigned Orderer = 0;
 
-	AutoRTFM::Commit([&]
+	AutoRTFM::Testing::Commit([&]
 	{
 		// If we are retrying transactions, need to reset the test state.
 		AutoRTFM::OnAbort([&]
@@ -106,7 +106,7 @@ TEST_CASE("Abort.NestedAbortOrder")
 
 TEST_CASE("Abort.TransactionInOnCommit")
 {
-	AutoRTFM::Commit([&]
+	AutoRTFM::Testing::Commit([&]
 	{
 		AutoRTFM::OnCommit([&]
 		{
@@ -586,7 +586,7 @@ TEST_CASE("Abort.PushOnAbortHandler_NoAbort")
 {
 	int Value = 55;
 
-	AutoRTFM::Commit([&]
+	AutoRTFM::Testing::Commit([&]
 	{
 		Value = 66;
 		AutoRTFM::PushOnAbortHandler(&SAbortKey, [&Value](){ Value = 77; });
@@ -631,7 +631,7 @@ TEST_CASE("Abort.PushOnAbortHandler_WithPop_NoAbort")
 {
 	int Value = 55;
 
-	AutoRTFM::Commit([&]
+	AutoRTFM::Testing::Commit([&]
 	{
 		Value = 66;
 		AutoRTFM::PushOnAbortHandler(&SAbortKey, [&Value]() { Value = 77; });
@@ -647,7 +647,7 @@ TEST_CASE("Abort.PushOnAbortHandler_WithPopAll_NoAbort")
 {
 	int Value = 55;
 
-	AutoRTFM::Commit([&]
+	AutoRTFM::Testing::Commit([&]
 	{
 		Value = 66;
 		AutoRTFM::PushOnAbortHandler(&SAbortKey, [&Value](){ Value = 77; });
@@ -819,7 +819,7 @@ TEST_CASE("Abort.PushOnAbortHandler_Order")
 				});
 
 				// Make a child transaction.
-				AutoRTFM::Commit([&]
+				AutoRTFM::Testing::Commit([&]
 				{
 					AutoRTFM::PushOnAbortHandler(&SAbortKey, [&Value]()
 					{
@@ -871,7 +871,7 @@ TEST_CASE("Abort.PushOnAbortHandler_Order")
 				AutoRTFM::OnAbort([&Value] { REQUIRE(40 == Value); Value += 1; });
 
 				// Make a child transaction.
-				AutoRTFM::Commit([&]
+				AutoRTFM::Testing::Commit([&]
 				{
 					AutoRTFM::PushOnAbortHandler(&SAbortKey, [&Value]()
 					{
@@ -899,7 +899,7 @@ TEST_CASE("Abort.PushOnAbortHandler_Order")
 			int Value = 99;
 
 			AutoRTFM::ETransactionResult Result = AutoRTFM::ETransactionResult::Committed;
-			AutoRTFM::Commit([&]
+			AutoRTFM::Testing::Commit([&]
 			{
 				AutoRTFM::OnCommit([&Value] { REQUIRE(37 == Value); Value += 1; });
 
@@ -932,7 +932,7 @@ TEST_CASE("Abort.PushOnAbortHandler_Order")
 			AutoRTFM::ETransactionResult Result = AutoRTFM::Transact([&]
 			{
 				AutoRTFM::PushOnAbortHandler(&SAbortKey, [&Value] { Value = 42; });
-				AutoRTFM::Commit([&] { AutoRTFM::PopOnAbortHandler(&SAbortKey); });
+				AutoRTFM::Testing::Commit([&] { AutoRTFM::PopOnAbortHandler(&SAbortKey); });
 				AutoRTFM::AbortTransaction();
 			});
 
@@ -970,7 +970,7 @@ TEST_CASE("Abort.PushOnAbortHandler_Order")
 			AutoRTFM::ETransactionResult Result = AutoRTFM::Transact([&]
 			{
 				AutoRTFM::PushOnAbortHandler(&SAbortKey, [&Value] { Value = 42; });
-				AutoRTFM::Commit([&] { AutoRTFM::Commit([&] { AutoRTFM::PopOnAbortHandler(&SAbortKey); }); });
+				AutoRTFM::Testing::Commit([&] { AutoRTFM::Testing::Commit([&] { AutoRTFM::PopOnAbortHandler(&SAbortKey); }); });
 				AutoRTFM::AbortTransaction();
 			});
 
@@ -984,7 +984,7 @@ TEST_CASE("Abort.PushOnAbortHandler_Order")
 			AutoRTFM::ETransactionResult Result = AutoRTFM::Transact([&]
 			{
 				AutoRTFM::PushOnAbortHandler(&SAbortKey, [&Value] { Value = 42; });
-				AutoRTFM::Commit([&]
+				AutoRTFM::Testing::Commit([&]
 				{
 					AutoRTFM::PushOnAbortHandler(&SAbortKey, [&Value] { Value = 13; });
 					AutoRTFM::PopAllOnAbortHandlers(&SAbortKey);
@@ -1025,7 +1025,7 @@ TEST_CASE("Abort.OnAbortTiming")
 	bool bOnPreAbortRan = false;
 	bool bOnAbortRan = false;
 	int Memory = 666;
-	AutoRTFM::Commit([&]
+	AutoRTFM::Testing::Commit([&]
 	{
 		// If we are retrying transactions, need to reset the test state.
 		AutoRTFM::OnPreAbort([&]
@@ -1230,7 +1230,7 @@ TEST_CASE("Abort.PushOnAbortOrdering")
 			AutoRTFM::OnAbort([&] { REQUIRE(14 == Orderer++); });
 
 			// This commits which will add the on abort handlers to the parent scope.
-			AutoRTFM::Commit([&]
+			AutoRTFM::Testing::Commit([&]
 			{
 				AutoRTFM::OnAbort([&] { REQUIRE(13 == Orderer++); });
 				AutoRTFM::PushOnAbortHandler(&Orderer, [&] { REQUIRE(12 == Orderer++); });
diff --git a/Engine/Source/Programs/AutoRTFMTests/Private/Allocators.cpp b/Engine/Source/Programs/AutoRTFMTests/Private/Allocators.cpp
--- a/Engine/Source/Programs/AutoRTFMTests/Private/Allocators.cpp
+++ b/Engine/Source/Programs/AutoRTFMTests/Private/Allocators.cpp
@@ -1,6 +1,7 @@
 // Copyright Epic Games, Inc. All Rights Reserved.
 
 #include "AutoRTFM.h"
+#include "AutoRTFM/Testing.h"
 #include "Catch2Includes.h"
 
 #include <new>
@@ -42,7 +43,7 @@ TEST_CASE("Allocators.New")
 	SECTION("With Commit")
 	{
 		int* Data = nullptr;
-		AutoRTFM::Commit([&]
+		AutoRTFM::Testing::Commit([&]
 			{
 				Data = new int(42);
 			});
@@ -75,7 +76,7 @@ TEST_CASE("Allocators.Delete")
 	{
 		int* Data = new int(42);
 
-		AutoRTFM::Commit([&]
+		AutoRTFM::Testing::Commit([&]
 			{
 				delete Data;
 			});
@@ -101,7 +102,7 @@ TEST_CASE("Allocators.ArrayNew")
 	SECTION("With Commit")
 	{
 		int* Data = nullptr;
-		AutoRTFM::Commit([&]
+		AutoRTFM::Testing::Commit([&]
 			{
 				Data = new int[42]();
 				Data[2] = 42;
@@ -136,7 +137,7 @@ TEST_CASE("Allocators.ArrayDelete")
 	{
 		int* Data = new int[42]();
 
-		AutoRTFM::Commit([&]
+		AutoRTFM::Testing::Commit([&]
 			{
 				delete[] Data;
 			});
@@ -169,7 +170,7 @@ TEST_CASE("Allocators.NewNoOpts")
 
 		CLANG_BEGIN_DISABLE_OPTIMIZATIONS
 
-		AutoRTFM::Commit([&]
+		AutoRTFM::Testing::Commit([&]
 			{
 				Data = new int(42);
 			});
@@ -209,7 +210,7 @@ TEST_CASE("Allocators.DeleteNoOpts")
 
 		CLANG_BEGIN_DISABLE_OPTIMIZATIONS
 
-		AutoRTFM::Commit([&]
+		AutoRTFM::Testing::Commit([&]
 			{
 				delete Data;
 			});
@@ -243,7 +244,7 @@ TEST_CASE("Allocators.Free")
 	{
 		void* Data = malloc(32);
 
-		AutoRTFM::Commit([&]
+		AutoRTFM::Testing::Commit([&]
 			{
 				free(Data);
 			});
@@ -273,7 +274,7 @@ TEST_CASE("Allocators.Malloc")
 	SECTION("With Commit")
 	{
 		void* Data = nullptr;
-		AutoRTFM::Commit([&]
+		AutoRTFM::Testing::Commit([&]
 			{
 				Data = malloc(32);
 				*reinterpret_cast<int*>(Data) = 42;
@@ -309,7 +310,7 @@ TEST_CASE("Allocators.Calloc")
 	SECTION("With Commit")
 	{
 		void* Data = nullptr;
-		AutoRTFM::Commit([&]
+		AutoRTFM::Testing::Commit([&]
 			{
 				Data = calloc(5, sizeof(int));
 				*reinterpret_cast<int*>(Data) = 42;
@@ -354,7 +355,7 @@ TEST_CASE("Allocators.Realloc")
 		void* Realloc = nullptr;
 		*reinterpret_cast<int*>(Alloc) = 42;
 
-		AutoRTFM::Commit([&]
+		AutoRTFM::Testing::Commit([&]
 			{
 				Realloc = realloc(Alloc, 64);
 			});
diff --git a/Engine/Source/Programs/AutoRTFMTests/Private/Benchmarks.cpp b/Engine/Source/Programs/AutoRTFMTests/Private/Benchmarks.cpp
--- a/Engine/Source/Programs/AutoRTFMTests/Private/Benchmarks.cpp
+++ b/Engine/Source/Programs/AutoRTFMTests/Private/Benchmarks.cpp
@@ -229,7 +229,7 @@ TEST_CASE("Graph")
 	{
 		FGraph<unsigned> Graph;
 
-		AutoRTFM::Commit([&]() { Graph.AddNode(42); });
+		AutoRTFM::Testing::Commit([&]() { Graph.AddNode(42); });
 
 		auto& Nodes = Graph.GetNodes();
 
@@ -240,7 +240,7 @@ TEST_CASE("Graph")
 	{
 		FGraph<unsigned> Graph;
 
-		AutoRTFM::Commit([&]() { Graph.AddRoot(Graph.AddNode(42)); });
+		AutoRTFM::Testing::Commit([&]() { Graph.AddRoot(Graph.AddNode(42)); });
 
 		auto& Roots = Graph.GetRoots();
 
@@ -269,7 +269,7 @@ TEST_CASE("Graph")
 			unsigned Result;
 			ResetXorshift();
 			FGraph<unsigned> Graph;
-			AutoRTFM::Commit([&]() { Graph = BuildGraph(Total); });
+			AutoRTFM::Testing::Commit([&]() { Graph = BuildGraph(Total); });
 			Result = WalkGraph(Graph);
 			CheckResult(Result, Total);
 		}
@@ -278,7 +278,7 @@ TEST_CASE("Graph")
 			unsigned Result;
 			ResetXorshift();
 			FGraph<unsigned> Graph = BuildGraph(Total);
-			AutoRTFM::Commit([&]() { Result = WalkGraph(Graph); });
+			AutoRTFM::Testing::Commit([&]() { Result = WalkGraph(Graph); });
 			CheckResult(Result, Total);
 		}
 
@@ -286,8 +286,8 @@ TEST_CASE("Graph")
 			unsigned Result;
 			ResetXorshift();
 			FGraph<unsigned> Graph;
-			AutoRTFM::Commit([&]() { Graph = BuildGraph(Total); });
-			AutoRTFM::Commit([&]() { Result = WalkGraph(Graph); });
+			AutoRTFM::Testing::Commit([&]() { Graph = BuildGraph(Total); });
+			AutoRTFM::Testing::Commit([&]() { Result = WalkGraph(Graph); });
 			CheckResult(Result, Total);
 		}
 
@@ -295,7 +295,7 @@ TEST_CASE("Graph")
 			unsigned Result;
 			ResetXorshift();
 			FGraph<unsigned> Graph;
-			AutoRTFM::Commit([&]() { Graph = BuildGraph(Total); Result = WalkGraph(Graph); });
+			AutoRTFM::Testing::Commit([&]() { Graph = BuildGraph(Total); Result = WalkGraph(Graph); });
 			CheckResult(Result, Total);
 		}
 	}
@@ -311,7 +311,7 @@ TEST_CASE("Graph")
 	{
 		ResetXorshift();
 		FGraph<unsigned> Graph;
-		AutoRTFM::Commit([&]() { Graph = BuildGraph(); });
+		AutoRTFM::Testing::Commit([&]() { Graph = BuildGraph(); });
 		WalkGraph(Graph);
 	};
 
@@ -319,22 +319,22 @@ TEST_CASE("Graph")
 	{
 		ResetXorshift();
 		FGraph<unsigned> Graph = BuildGraph();
-		AutoRTFM::Commit([&]() { WalkGraph(Graph); });
+		AutoRTFM::Testing::Commit([&]() { WalkGraph(Graph); });
 	};
 
 	BENCHMARK("build transactional / walk transactional")
 	{
 		ResetXorshift();
 		FGraph<unsigned> Graph;
-		AutoRTFM::Commit([&]() { Graph = BuildGraph(); });
-		AutoRTFM::Commit([&]() { WalkGraph(Graph); });
+		AutoRTFM::Testing::Commit([&]() { Graph = BuildGraph(); });
+		AutoRTFM::Testing::Commit([&]() { WalkGraph(Graph); });
 	};
 
 	BENCHMARK("build + walk transactional")
 	{
 		ResetXorshift();
 		FGraph<unsigned> Graph;
-		AutoRTFM::Commit([&]() { Graph = BuildGraph(); WalkGraph(Graph); });
+		AutoRTFM::Testing::Commit([&]() { Graph = BuildGraph(); WalkGraph(Graph); });
 	};
 }
 
diff --git a/Engine/Source/Programs/AutoRTFMTests/Private/CString.cpp b/Engine/Source/Programs/AutoRTFMTests/Private/CString.cpp
--- a/Engine/Source/Programs/AutoRTFMTests/Private/CString.cpp
+++ b/Engine/Source/Programs/AutoRTFMTests/Private/CString.cpp
@@ -77,7 +77,7 @@ void CheckWCharFunction(FUNCTION&& Function, wchar_t WChar)
 		int Got = 0;
 
 		CLANG_BEGIN_DISABLE_OPTIMIZATIONS
-		AutoRTFM::Commit([&] { Got = Function(WChar); });
+		AutoRTFM::Testing::Commit([&] { Got = Function(WChar); });
 		CLANG_END_DISABLE_OPTIMIZATIONS
 
 		REQUIRE(Expect == Got);
@@ -120,7 +120,7 @@ void CheckStringToNumberFunction(RESULT Function(const CHAR*, CHAR**, ADDITIONAL
 		SECTION("With Commit")
 		{
 			CLANG_BEGIN_DISABLE_OPTIMIZATIONS
-			AutoRTFM::Commit([&]
+			AutoRTFM::Testing::Commit([&]
 			{
 				GotValue = Function(String, &GotEndPtr, AdditionalArgs...);
 			});
@@ -149,7 +149,7 @@ void CheckStringToNumberFunction(RESULT Function(const CHAR*, CHAR**, ADDITIONAL
 		SECTION("With Commit")
 		{
 			CLANG_BEGIN_DISABLE_OPTIMIZATIONS
-			AutoRTFM::Commit([&] { GotValue = Function(String, nullptr, AdditionalArgs...); });
+			AutoRTFM::Testing::Commit([&] { GotValue = Function(String, nullptr, AdditionalArgs...); });
 			CLANG_END_DISABLE_OPTIMIZATIONS
 
 			REQUIRE(ExpectValue == GotValue);
@@ -233,7 +233,7 @@ TEST_CASE("CString.memcpy")
 		SECTION("With Commit")
 		{
 			CLANG_BEGIN_DISABLE_OPTIMIZATIONS
-			AutoRTFM::Commit([&] { memcpy(To, From, 6); });
+			AutoRTFM::Testing::Commit([&] { memcpy(To, From, 6); });
 			CLANG_END_DISABLE_OPTIMIZATIONS
 
 			REQUIRE("Kittie says woof" == std::string_view(To));
@@ -262,7 +262,7 @@ TEST_CASE("CString.memcpy")
 		SECTION("With Commit")
 		{
 			CLANG_BEGIN_DISABLE_OPTIMIZATIONS
-			AutoRTFM::Commit([&] { Fn(To, From, 6); });
+			AutoRTFM::Testing::Commit([&] { Fn(To, From, 6); });
 			CLANG_END_DISABLE_OPTIMIZATIONS
 
 			REQUIRE("Kittie says woof" == std::string_view(To));
@@ -293,7 +293,7 @@ TEST_CASE("CString.memmove")
 		SECTION("With Commit")
 		{
 			CLANG_BEGIN_DISABLE_OPTIMIZATIONS
-			AutoRTFM::Commit([&] { memmove(To + 7, To, 5); });
+			AutoRTFM::Testing::Commit([&] { memmove(To + 7, To, 5); });
 			CLANG_END_DISABLE_OPTIMIZATIONS
 
 			REQUIRE("Hello, Hello!" == std::string_view(To));
@@ -322,7 +322,7 @@ TEST_CASE("CString.memmove")
 		SECTION("With Commit")
 		{
 			CLANG_BEGIN_DISABLE_OPTIMIZATIONS
-			AutoRTFM::Commit([&] { Fn(To + 7, To, 5); });
+			AutoRTFM::Testing::Commit([&] { Fn(To + 7, To, 5); });
 			CLANG_END_DISABLE_OPTIMIZATIONS
 
 			REQUIRE("Hello, Hello!" == std::string_view(To));
@@ -353,7 +353,7 @@ TEST_CASE("CString.memset")
 		SECTION("With Commit")
 		{
 			CLANG_BEGIN_DISABLE_OPTIMIZATIONS
-			AutoRTFM::Commit([&] { memset(To + 7, 'X', 5); });
+			AutoRTFM::Testing::Commit([&] { memset(To + 7, 'X', 5); });
 			CLANG_END_DISABLE_OPTIMIZATIONS
 
 			REQUIRE("Hello, XXXXX!" == std::string_view(To));
@@ -382,7 +382,7 @@ TEST_CASE("CString.memset")
 		SECTION("With Commit")
 		{
 			CLANG_BEGIN_DISABLE_OPTIMIZATIONS
-			AutoRTFM::Commit([&] { Fn(To + 7, 'X', 5); });
+			AutoRTFM::Testing::Commit([&] { Fn(To + 7, 'X', 5); });
 			CLANG_END_DISABLE_OPTIMIZATIONS
 
 			REQUIRE("Hello, XXXXX!" == std::string_view(To));
@@ -415,7 +415,7 @@ TEST_CASE("CString.strcpy")
 	SECTION("With Commit")
 	{
 		CLANG_BEGIN_DISABLE_OPTIMIZATIONS
-		AutoRTFM::Commit([&] { strcpy(To, From); });
+		AutoRTFM::Testing::Commit([&] { strcpy(To, From); });
 		CLANG_END_DISABLE_OPTIMIZATIONS
 
 		REQUIRE("Kittie says meow" == std::string_view(To));
@@ -448,7 +448,7 @@ TEST_CASE("CString.strncpy")
 	SECTION("With Commit")
 	{
 		CLANG_BEGIN_DISABLE_OPTIMIZATIONS
-		AutoRTFM::Commit([&] { strncpy(To, From, 6); });
+		AutoRTFM::Testing::Commit([&] { strncpy(To, From, 6); });
 		CLANG_END_DISABLE_OPTIMIZATIONS
 
 		REQUIRE("Kittie says woof" == std::string_view(To));
@@ -481,7 +481,7 @@ TEST_CASE("CString.strcat")
 	SECTION("With Commit")
 	{
 		CLANG_BEGIN_DISABLE_OPTIMIZATIONS
-		AutoRTFM::Commit([&] { strcat(To, ", world!"); });
+		AutoRTFM::Testing::Commit([&] { strcat(To, ", world!"); });
 		CLANG_END_DISABLE_OPTIMIZATIONS
 
 		REQUIRE("Hello, world!" == std::string_view(To));
@@ -514,7 +514,7 @@ TEST_CASE("CString.strncat")
 	SECTION("With Commit")
 	{
 		CLANG_BEGIN_DISABLE_OPTIMIZATIONS
-		AutoRTFM::Commit([&] { strncat(To, ", world! Not this!", 8); });
+		AutoRTFM::Testing::Commit([&] { strncat(To, ", world! Not this!", 8); });
 		CLANG_END_DISABLE_OPTIMIZATIONS
 
 		REQUIRE("Hello, world!" == std::string_view(To));
@@ -549,7 +549,7 @@ TEST_CASE("CString.memcmp")
 		int Compare = 0;
 
 		CLANG_BEGIN_DISABLE_OPTIMIZATIONS
-		AutoRTFM::Commit([&] { Compare = memcmp(A, "That", 4); });
+		AutoRTFM::Testing::Commit([&] { Compare = memcmp(A, "That", 4); });
 		CLANG_END_DISABLE_OPTIMIZATIONS
 
 		REQUIRE(0 < Compare);
@@ -581,7 +581,7 @@ TEST_CASE("CString.strcmp")
 		int Compare = 0;
 
 		CLANG_BEGIN_DISABLE_OPTIMIZATIONS
-		AutoRTFM::Commit([&] { Compare = strcmp(A, "That"); });
+		AutoRTFM::Testing::Commit([&] { Compare = strcmp(A, "That"); });
 		CLANG_END_DISABLE_OPTIMIZATIONS
 
 		REQUIRE(0 < Compare);
@@ -613,7 +613,7 @@ TEST_CASE("CString.strncmp")
 		int Compare = 0;
 
 		CLANG_BEGIN_DISABLE_OPTIMIZATIONS
-		AutoRTFM::Commit([&] { Compare = strncmp(A, "That", 3); });
+		AutoRTFM::Testing::Commit([&] { Compare = strncmp(A, "That", 3); });
 		CLANG_END_DISABLE_OPTIMIZATIONS
 
 		REQUIRE(0 < Compare);
@@ -645,7 +645,7 @@ TEST_CASE("CString.strchr")
 		const char* Value = nullptr;
 
 		CLANG_BEGIN_DISABLE_OPTIMIZATIONS
-		AutoRTFM::Commit([&] { Value = strchr(A, 'i'); });
+		AutoRTFM::Testing::Commit([&] { Value = strchr(A, 'i'); });
 		CLANG_END_DISABLE_OPTIMIZATIONS
 
 		REQUIRE((A + 2) == Value);
@@ -677,7 +677,7 @@ TEST_CASE("CString.strrchr")
 		const char* Value = nullptr;
 
 		CLANG_BEGIN_DISABLE_OPTIMIZATIONS
-		AutoRTFM::Commit([&] { Value = strrchr(A, 'i'); });
+		AutoRTFM::Testing::Commit([&] { Value = strrchr(A, 'i'); });
 		CLANG_END_DISABLE_OPTIMIZATIONS
 
 		REQUIRE((A + 5) == Value);
@@ -709,7 +709,7 @@ TEST_CASE("CString.strstr")
 		const char* Value = nullptr;
 
 		CLANG_BEGIN_DISABLE_OPTIMIZATIONS
-		AutoRTFM::Commit([&] { Value = strstr(A, "is"); });
+		AutoRTFM::Testing::Commit([&] { Value = strstr(A, "is"); });
 		CLANG_END_DISABLE_OPTIMIZATIONS
 
 		REQUIRE((A + 2) == Value);
@@ -741,7 +741,7 @@ TEST_CASE("CString.strlen")
 		size_t Value = 0;
 
 		CLANG_BEGIN_DISABLE_OPTIMIZATIONS
-		AutoRTFM::Commit([&] { Value = strlen(A); });
+		AutoRTFM::Testing::Commit([&] { Value = strlen(A); });
 		CLANG_END_DISABLE_OPTIMIZATIONS
 
 		REQUIRE(4 == Value);
@@ -773,7 +773,7 @@ TEST_CASE("CString.wcscpy")
 	SECTION("With Commit")
 	{
 		CLANG_BEGIN_DISABLE_OPTIMIZATIONS
-		AutoRTFM::Commit([&] { wcscpy(To, From); });
+		AutoRTFM::Testing::Commit([&] { wcscpy(To, From); });
 		CLANG_END_DISABLE_OPTIMIZATIONS
 
 		REQUIRE(std::wstring_view(L"Kittie says meow") == std::wstring_view(To));
@@ -806,7 +806,7 @@ TEST_CASE("CString.wcsncpy")
 	SECTION("With Commit")
 	{
 		CLANG_BEGIN_DISABLE_OPTIMIZATIONS
-		AutoRTFM::Commit([&] { wcsncpy(To, From, 6); });
+		AutoRTFM::Testing::Commit([&] { wcsncpy(To, From, 6); });
 		CLANG_END_DISABLE_OPTIMIZATIONS
 
 		REQUIRE(std::wstring_view(L"Kittie says woof") == std::wstring_view(To));
@@ -840,7 +840,7 @@ TEST_CASE("CString.wcsncmp")
 		int Compare = 0;
 
 		CLANG_BEGIN_DISABLE_OPTIMIZATIONS
-		AutoRTFM::Commit([&] { Compare = wcsncmp(A, L"That", 3); });
+		AutoRTFM::Testing::Commit([&] { Compare = wcsncmp(A, L"That", 3); });
 		CLANG_END_DISABLE_OPTIMIZATIONS
 
 		REQUIRE(0 < Compare);
@@ -874,7 +874,7 @@ TEST_CASE("CString.wcschr")
 			const wchar_t* Value = nullptr;
 
 			CLANG_BEGIN_DISABLE_OPTIMIZATIONS
-			AutoRTFM::Commit([&] { Value = wcschr(A, L'i'); });
+			AutoRTFM::Testing::Commit([&] { Value = wcschr(A, L'i'); });
 			CLANG_END_DISABLE_OPTIMIZATIONS
 
 			REQUIRE((A + 2) == Value);
@@ -906,7 +906,7 @@ TEST_CASE("CString.wcschr")
 			wchar_t* Value = nullptr;
 
 			CLANG_BEGIN_DISABLE_OPTIMIZATIONS
-			AutoRTFM::Commit([&] { Value = wcschr(A, L'i'); });
+			AutoRTFM::Testing::Commit([&] { Value = wcschr(A, L'i'); });
 			CLANG_END_DISABLE_OPTIMIZATIONS
 
 			REQUIRE((A + 2) == Value);
@@ -941,7 +941,7 @@ TEST_CASE("CString.wcsstr")
 			const wchar_t* Value = nullptr;
 
 			CLANG_BEGIN_DISABLE_OPTIMIZATIONS
-			AutoRTFM::Commit([&] { Value = wcsstr(A, L"ink"); });
+			AutoRTFM::Testing::Commit([&] { Value = wcsstr(A, L"ink"); });
 			CLANG_END_DISABLE_OPTIMIZATIONS
 
 			REQUIRE((A + 2) == Value);
@@ -973,7 +973,7 @@ TEST_CASE("CString.wcsstr")
 			wchar_t* Value = nullptr;
 
 			CLANG_BEGIN_DISABLE_OPTIMIZATIONS
-			AutoRTFM::Commit([&] { Value = wcsstr(A, L"ink"); });
+			AutoRTFM::Testing::Commit([&] { Value = wcsstr(A, L"ink"); });
 			CLANG_END_DISABLE_OPTIMIZATIONS
 
 			REQUIRE((A + 2) == Value);
@@ -1005,7 +1005,7 @@ TEST_CASE("CString.wcsrchr")
 		const wchar_t* Return = nullptr;
 
 		CLANG_BEGIN_DISABLE_OPTIMIZATIONS
-		AutoRTFM::Commit([&] { Return = wcsrchr(A, 'h'); });
+		AutoRTFM::Testing::Commit([&] { Return = wcsrchr(A, 'h'); });
 		CLANG_END_DISABLE_OPTIMIZATIONS
 
 		REQUIRE(nullptr != Return);
@@ -1038,7 +1038,7 @@ TEST_CASE("CString.wcscmp")
 		int Compare = 0;
 
 		CLANG_BEGIN_DISABLE_OPTIMIZATIONS
-		AutoRTFM::Commit([&] { Compare = wcscmp(A, L"That"); });
+		AutoRTFM::Testing::Commit([&] { Compare = wcscmp(A, L"That"); });
 		CLANG_END_DISABLE_OPTIMIZATIONS
 
 		REQUIRE(0 < Compare);
@@ -1070,7 +1070,7 @@ TEST_CASE("CString.wcslen")
 		size_t Value = 0;
 
 		CLANG_BEGIN_DISABLE_OPTIMIZATIONS
-		AutoRTFM::Commit([&] { Value = wcslen(A); });
+		AutoRTFM::Testing::Commit([&] { Value = wcslen(A); });
 		CLANG_END_DISABLE_OPTIMIZATIONS
 
 		REQUIRE(4 == Value);
@@ -1309,7 +1309,7 @@ TEST_CASE("CString.snprintf")
 			int Count = 0;
 
 			CLANG_BEGIN_DISABLE_OPTIMIZATIONS
-			AutoRTFM::Commit([&]
+			AutoRTFM::Testing::Commit([&]
 			{
 				Count = snprintf(Buffer, BufferSize, "cat says '%s'!", "meow");
 			});
@@ -1346,7 +1346,7 @@ TEST_CASE("CString.snprintf")
 			int Count = 0;
 
 			CLANG_BEGIN_DISABLE_OPTIMIZATIONS
-			AutoRTFM::Commit([&]
+			AutoRTFM::Testing::Commit([&]
 			{
 				Count = snprintf(Buffer, 8, "cat says '%s'!", "meow");
 			});
@@ -1382,7 +1382,7 @@ TEST_CASE("CString.snprintf")
 			int Count = 0;
 
 			CLANG_BEGIN_DISABLE_OPTIMIZATIONS
-			AutoRTFM::Commit([&]
+			AutoRTFM::Testing::Commit([&]
 			{
 				Count = snprintf(nullptr, 0, "cat says '%s'!", "meow");
 			});
@@ -1435,7 +1435,7 @@ TEST_CASE("CString.snprintf")
 			int Count = 0;
 
 			CLANG_BEGIN_DISABLE_OPTIMIZATIONS
-			AutoRTFM::Commit([&]
+			AutoRTFM::Testing::Commit([&]
 			{
 				Count = snprintf(Buffer, 8, "%%n");
 			});
@@ -1476,7 +1476,7 @@ TEST_CASE("CString.swprintf")
 			int Count = 0;
 
 			CLANG_BEGIN_DISABLE_OPTIMIZATIONS
-			AutoRTFM::Commit([&]
+			AutoRTFM::Testing::Commit([&]
 			{
 				Count = swprintf(Buffer, BufferSize, L"cat says '%ls'!", L"meow");
 			});
@@ -1513,7 +1513,7 @@ TEST_CASE("CString.swprintf")
 			int Count = 0;
 
 			CLANG_BEGIN_DISABLE_OPTIMIZATIONS
-			AutoRTFM::Commit([&]
+			AutoRTFM::Testing::Commit([&]
 			{
 				Count = swprintf(Buffer, 8, L"cat says '%ls'!", L"meow");
 			});
@@ -1549,7 +1549,7 @@ TEST_CASE("CString.swprintf")
 			int Count = 0;
 
 			CLANG_BEGIN_DISABLE_OPTIMIZATIONS
-			AutoRTFM::Commit([&]
+			AutoRTFM::Testing::Commit([&]
 			{
 				Count = swprintf(nullptr, 0, L"cat says '%ls'!", L"meow");
 			});
@@ -1602,7 +1602,7 @@ TEST_CASE("CString.swprintf")
 			int Count = 0;
 
 			CLANG_BEGIN_DISABLE_OPTIMIZATIONS
-			AutoRTFM::Commit([&]
+			AutoRTFM::Testing::Commit([&]
 			{
 				Count = swprintf(Buffer, 8, L"%%n");
 			});
@@ -1637,7 +1637,7 @@ TEST_CASE("CString.printf")
 		int Count = 0;
 
 		CLANG_BEGIN_DISABLE_OPTIMIZATIONS
-		AutoRTFM::Commit([&]
+		AutoRTFM::Testing::Commit([&]
 		{
 			Count = printf("AutoRTFM printf test: %%n\n");
 		});
diff --git a/Engine/Source/Programs/AutoRTFMTests/Private/Commit.cpp b/Engine/Source/Programs/AutoRTFMTests/Private/Commit.cpp
--- a/Engine/Source/Programs/AutoRTFMTests/Private/Commit.cpp
+++ b/Engine/Source/Programs/AutoRTFMTests/Private/Commit.cpp
@@ -70,7 +70,7 @@ TEST_CASE("Commit.Large")
             m[i + j].push_back(i + j);
         }
     }
-    AutoRTFM::Commit([&]
+    AutoRTFM::Testing::Commit([&]
     {
         for (unsigned i = 0; i < 10000; ++i)
         {
@@ -137,7 +137,7 @@ TEST_CASE("Commit.PushOnCommitHandler_WithPop_NoAbort")
 
 TEST_CASE("Commit.PushOnCommitHandler_WithPopAll_NoAbort")
 {
-	AutoRTFM::Commit([&]
+	AutoRTFM::Testing::Commit([&]
 	{
 		AutoRTFM::PushOnCommitHandler(&SCommitKey, []() { FAIL("Unreachable!"); });
 		AutoRTFM::PopAllOnCommitHandlers(&SCommitKey);
diff --git a/Engine/Source/Programs/AutoRTFMTests/Private/CopyAddInt64.cpp b/Engine/Source/Programs/AutoRTFMTests/Private/CopyAddInt64.cpp
deleted file mode 100644
diff --git a/Engine/Source/Programs/AutoRTFMTests/Private/DependencyTests.cpp b/Engine/Source/Programs/AutoRTFMTests/Private/DependencyTests.cpp
--- a/Engine/Source/Programs/AutoRTFMTests/Private/DependencyTests.cpp
+++ b/Engine/Source/Programs/AutoRTFMTests/Private/DependencyTests.cpp
@@ -1,6 +1,7 @@
 // Copyright Epic Games, Inc. All Rights Reserved.
 
 #include "AutoRTFM.h"
+#include "AutoRTFM/Testing.h"
 
 // We can't use Catch2Includes.h here, because it #includes more Unreal headers.
 // Fortunately, we don't even need to run these tests; if this code compiles and links, 
@@ -31,7 +32,7 @@ void OnAbort_FunctionWorksWithNoAdditionalIncludes()
 
 void OnCommit_MacroWorksWithNoAdditionalIncludes()
 {
-	AutoRTFM::Commit([&]
+	AutoRTFM::Testing::Commit([&]
 	{
 		UE_AUTORTFM_ONCOMMIT()
 		{
@@ -41,7 +42,7 @@ void OnCommit_MacroWorksWithNoAdditionalIncludes()
 
 void OnCommit_FunctionWorksWithNoAdditionalIncludes()
 {
-	AutoRTFM::Commit([&]
+	AutoRTFM::Testing::Commit([&]
 	{
 		AutoRTFM::OnCommit([]
 		{
diff --git a/Engine/Source/Programs/AutoRTFMTests/Private/FunctionPointerTests.cpp b/Engine/Source/Programs/AutoRTFMTests/Private/FunctionPointerTests.cpp
--- a/Engine/Source/Programs/AutoRTFMTests/Private/FunctionPointerTests.cpp
+++ b/Engine/Source/Programs/AutoRTFMTests/Private/FunctionPointerTests.cpp
@@ -1,6 +1,7 @@
 // Copyright Epic Games, Inc. All Rights Reserved.
 
 #include "AutoRTFM.h"
+#include "AutoRTFM/Testing.h"
 #include "Catch2Includes.h"
 #include "Templates/Function.h"
 
@@ -40,7 +41,7 @@ namespace
 
 TEST_CASE("FunctionPointer.CStyle")
 {
-	AutoRTFM::Commit([&]
+	AutoRTFM::Testing::Commit([&]
 		{
 			CStyleType CStyle = GetTestCFunction();
 			REQUIRE(42 == CStyle());
@@ -51,7 +52,7 @@ TEST_CASE("FunctionPointer.StandardLibrary")
 {
 	SECTION("Created inside transaction")
 	{
-		AutoRTFM::Commit([&]
+		AutoRTFM::Testing::Commit([&]
 			{
 				CosfType fn = GetCosfFunction();
 				REQUIRE(1.0f == fn(0.0f));
@@ -62,7 +63,7 @@ TEST_CASE("FunctionPointer.StandardLibrary")
 	{
 		CosfType fn = GetCosfFunction();
 
-		AutoRTFM::Commit([&]
+		AutoRTFM::Testing::Commit([&]
 			{
 				REQUIRE(1.0f == fn(0.0f));
 			});
@@ -74,7 +75,7 @@ TEST_CASE("FunctionPointer.TFunction")
 	SECTION("Created inside transaction")
 	{
 		int Result = 0;
-		AutoRTFM::Commit([&]
+		AutoRTFM::Testing::Commit([&]
 			{
 				TFunction<void()> MyFunc = [&Result]
 					{
@@ -102,7 +103,7 @@ TEST_CASE("FunctionPointer.TFunction")
 				Result = 42;
 			};
 
-		AutoRTFM::Commit([&]
+		AutoRTFM::Testing::Commit([&]
 			{
 				if (MyFunc)
 				{
@@ -123,7 +124,7 @@ TEST_CASE("FunctionPointer.TUniqueFunction")
 	SECTION("Created inside transaction")
 	{
 		int Result = 0;
-		AutoRTFM::Commit([&]
+		AutoRTFM::Testing::Commit([&]
 			{
 				TUniqueFunction<void()> MyFunc = [&Result]
 					{
@@ -151,7 +152,7 @@ TEST_CASE("FunctionPointer.TUniqueFunction")
 				Result = 42;
 			};
 
-		AutoRTFM::Commit([&]
+		AutoRTFM::Testing::Commit([&]
 			{
 				if (MyFunc)
 				{
diff --git a/Engine/Source/Programs/AutoRTFMTests/Private/LockFreeFixedSizeAllocatorTests.cpp b/Engine/Source/Programs/AutoRTFMTests/Private/LockFreeFixedSizeAllocatorTests.cpp
--- a/Engine/Source/Programs/AutoRTFMTests/Private/LockFreeFixedSizeAllocatorTests.cpp
+++ b/Engine/Source/Programs/AutoRTFMTests/Private/LockFreeFixedSizeAllocatorTests.cpp
@@ -226,7 +226,7 @@ TEMPLATE_TEST_CASE("LockFreeFixedSizeAllocator", "", FThreadSafeCounter, FThread
 		ALockFreeFixedSizeAllocator Allocator;
 		IntegerType NumFree = 0;
 
-		AutoRTFM::Commit([&]
+		AutoRTFM::Testing::Commit([&]
 		{
 			NumFree = Allocator.GetNumFree();
 		});
@@ -236,7 +236,7 @@ TEMPLATE_TEST_CASE("LockFreeFixedSizeAllocator", "", FThreadSafeCounter, FThread
 		void* const Blob = Allocator.Allocate();
 		Allocator.Free(Blob);
 
-		AutoRTFM::Commit([&]
+		AutoRTFM::Testing::Commit([&]
 		{
 			NumFree = Allocator.GetNumFree();
 		});
@@ -249,7 +249,7 @@ TEMPLATE_TEST_CASE("LockFreeFixedSizeAllocator", "", FThreadSafeCounter, FThread
 		ALockFreeFixedSizeAllocator Allocator;
 		IntegerType NumUsed = 0;
 
-		AutoRTFM::Commit([&]
+		AutoRTFM::Testing::Commit([&]
 		{
 			NumUsed = Allocator.GetNumUsed();
 		});
@@ -258,7 +258,7 @@ TEMPLATE_TEST_CASE("LockFreeFixedSizeAllocator", "", FThreadSafeCounter, FThread
 
 		void* const Blob = Allocator.Allocate();
 
-		AutoRTFM::Commit([&]
+		AutoRTFM::Testing::Commit([&]
 		{
 			NumUsed = Allocator.GetNumUsed();
 		});
diff --git a/Engine/Source/Programs/AutoRTFMTests/Private/Math.cpp b/Engine/Source/Programs/AutoRTFMTests/Private/Math.cpp
--- a/Engine/Source/Programs/AutoRTFMTests/Private/Math.cpp
+++ b/Engine/Source/Programs/AutoRTFMTests/Private/Math.cpp
@@ -1,6 +1,7 @@
 // Copyright Epic Games, Inc. All Rights Reserved.
 
 #include "AutoRTFM.h"
+#include "AutoRTFM/Testing.h"
 #include "Catch2Includes.h"
 
 #include <math.h>
@@ -38,7 +39,7 @@ void Check(AUTORTFM_IMPLICIT_ENABLE FUNC&& Function)
 	{
 		T Got = Zero;
 
-		AutoRTFM::Commit([&] { Got = Function(); });
+		AutoRTFM::Testing::Commit([&] { Got = Function(); });
 
 		REQUIRE(Expected == Got);
 	}
@@ -318,6 +319,6 @@ TEST_CASE("Math.rand")
 
 	SECTION("With Commit")
 	{
-		AutoRTFM::Commit([&] { rand(); });
+		AutoRTFM::Testing::Commit([&] { rand(); });
 	}
 }
diff --git a/Engine/Source/Programs/AutoRTFMTests/Private/Open.cpp b/Engine/Source/Programs/AutoRTFMTests/Private/Open.cpp
--- a/Engine/Source/Programs/AutoRTFMTests/Private/Open.cpp
+++ b/Engine/Source/Programs/AutoRTFMTests/Private/Open.cpp
@@ -137,7 +137,7 @@ TEST_CASE("Open.FunctionPtrFromAlwaysOpenFunction")
 
 	SECTION("Commit")
 	{
-		AutoRTFM::Commit([&]
+		AutoRTFM::Testing::Commit([&]
 		{
 			Func(&Value);
 			REQUIRE(Value == 42);
diff --git a/Engine/Source/Programs/AutoRTFMTests/Private/PackedPointer.cpp b/Engine/Source/Programs/AutoRTFMTests/Private/PackedPointer.cpp
--- a/Engine/Source/Programs/AutoRTFMTests/Private/PackedPointer.cpp
+++ b/Engine/Source/Programs/AutoRTFMTests/Private/PackedPointer.cpp
@@ -1,6 +1,7 @@
 // Copyright Epic Games, Inc. All Rights Reserved.
 
 #include "AutoRTFM.h"
+#include "AutoRTFM/Testing.h"
 #include "Catch2Includes.h"
 
 #include <inttypes.h>
@@ -18,7 +19,7 @@ TEST_CASE("PackedPointer")
     Foo foo;
     foo.a = 1;
     foo.b = 2;
-    AutoRTFM::Commit([&] ()
+    AutoRTFM::Testing::Commit([&] ()
     {
         foo.a++;
         foo.b++;
diff --git a/Engine/Source/Programs/AutoRTFMTests/Private/SharedPointerTests.cpp b/Engine/Source/Programs/AutoRTFMTests/Private/SharedPointerTests.cpp
--- a/Engine/Source/Programs/AutoRTFMTests/Private/SharedPointerTests.cpp
+++ b/Engine/Source/Programs/AutoRTFMTests/Private/SharedPointerTests.cpp
@@ -1,6 +1,7 @@
 // Copyright Epic Games, Inc. All Rights Reserved.
 
 #include "AutoRTFM.h"
+#include "AutoRTFM/Testing.h"
 #include "Catch2Includes.h"
 #include "HAL/UnrealMemory.h"
 #include "Templates/SharedPointer.h"
@@ -9,7 +10,7 @@ TEST_CASE("SharedPointer.PreviouslyAllocated")
 {
 	TSharedPtr<int, ESPMode::ThreadSafe> Foo(new int(42));
 
-	AutoRTFM::Commit([&]
+	AutoRTFM::Testing::Commit([&]
 		{
 			// Make a copy to bump the reference count.
 			TSharedPtr<int, ESPMode::ThreadSafe> Copy = Foo;
@@ -43,7 +44,7 @@ TEST_CASE("SharedPointer.NewlyAllocated")
 {
 	int Copy = 42;
 
-	AutoRTFM::Commit([&]
+	AutoRTFM::Testing::Commit([&]
 		{
 			TSharedPtr<int, ESPMode::ThreadSafe> Foo(new int(13));
 
@@ -72,9 +73,9 @@ TEST_CASE("SharedPointer.NestedTransactionWithPreviouslyAllocated")
 {
 	TSharedPtr<int, ESPMode::ThreadSafe> Foo(new int(42));
 
-	AutoRTFM::Commit([&]
+	AutoRTFM::Testing::Commit([&]
 		{
-			AutoRTFM::Commit([&]
+			AutoRTFM::Testing::Commit([&]
 				{
 					// Make a copy to bump the reference count.
 					TSharedPtr<int, ESPMode::ThreadSafe> Copy = Foo;
@@ -91,7 +92,7 @@ TEST_CASE("SharedPointer.AbortNestedTransactionWithPreviouslyAllocated")
 {
 	TSharedPtr<int, ESPMode::ThreadSafe> Foo(new int(42));
 
-	AutoRTFM::Commit([&]
+	AutoRTFM::Testing::Commit([&]
 		{
 			REQUIRE(AutoRTFM::ETransactionResult::AbortedByRequest == AutoRTFM::Transact([&]
 				{
@@ -112,9 +113,9 @@ TEST_CASE("SharedPointer.NestedTransactionWithNewlyAllocated")
 {
 	int Result = 42;
 
-	AutoRTFM::Commit([&]
+	AutoRTFM::Testing::Commit([&]
 		{
-			AutoRTFM::Commit([&]
+			AutoRTFM::Testing::Commit([&]
 				{
 					TSharedPtr<int, ESPMode::ThreadSafe> Foo(new int(13));
 					TSharedPtr<int, ESPMode::ThreadSafe> Copy = Foo;
@@ -129,7 +130,7 @@ TEST_CASE("SharedPointer.AbortNestedTransactionWithNewlyAllocated")
 {
 	int Result = 42;
 
-	AutoRTFM::Commit([&]
+	AutoRTFM::Testing::Commit([&]
 		{
 			REQUIRE(AutoRTFM::ETransactionResult::AbortedByRequest == AutoRTFM::Transact([&]
 				{
@@ -152,9 +153,9 @@ TEST_CASE("SharedPointer.NestedTransactionWithPlacementNewlyAllocated")
 {
 	int Result = 42;
 
-	AutoRTFM::Commit([&]
+	AutoRTFM::Testing::Commit([&]
 		{
-			AutoRTFM::Commit([&]
+			AutoRTFM::Testing::Commit([&]
 				{
 					using Ptr = TSharedPtr<int, ESPMode::ThreadSafe>;
 					void* const Memory = MakeMemoryForT<Ptr>();
@@ -173,7 +174,7 @@ TEST_CASE("SharedPointer.AbortNestedTransactionWithPlacementNewlyAllocated")
 {
 	int Result = 42;
 
-	AutoRTFM::Commit([&]
+	AutoRTFM::Testing::Commit([&]
 		{
 			REQUIRE(AutoRTFM::ETransactionResult::AbortedByRequest == AutoRTFM::Transact([&]
 				{
diff --git a/Engine/Source/Programs/AutoRTFMTests/Private/Tests.cpp b/Engine/Source/Programs/AutoRTFMTests/Private/Tests.cpp
--- a/Engine/Source/Programs/AutoRTFMTests/Private/Tests.cpp
+++ b/Engine/Source/Programs/AutoRTFMTests/Private/Tests.cpp
@@ -31,7 +31,7 @@ TEST_CASE("Tests.WriteInt")
 
 TEST_CASE("Tests.UE_LOG")
 {
-	AutoRTFM::Commit([&]
+	AutoRTFM::Testing::Commit([&]
 	{
 		UE_LOG(LogAutoRTFMTests, Display, TEXT("Testing this works!"));
 	});
@@ -336,7 +336,7 @@ TEST_CASE("Tests.RetryNonNested")
 
 	unsigned Count = 0;
 
-	AutoRTFM::Commit([&]
+	AutoRTFM::Testing::Commit([&]
 	{
 		AutoRTFM::Open([&]
 		{
@@ -468,7 +468,7 @@ TEST_CASE("Tests.__local_stdio_printf_options")
 	unsigned __int64* NonTransactional = __local_stdio_printf_options();
 	unsigned __int64* Transactional = nullptr;
 
-	AutoRTFM::Commit([&]
+	AutoRTFM::Testing::Commit([&]
 		{
 			Transactional = __local_stdio_printf_options();
 		});
@@ -481,7 +481,7 @@ TEST_CASE("Tests.__local_stdio_scanf_options")
 	unsigned __int64* NonTransactional = __local_stdio_scanf_options();
 	unsigned __int64* Transactional = nullptr;
 
-	AutoRTFM::Commit([&]
+	AutoRTFM::Testing::Commit([&]
 		{
 			Transactional = __local_stdio_scanf_options();
 		});
diff --git a/Engine/Source/Programs/AutoRTFMTests/Private/UnorderedSet.cpp b/Engine/Source/Programs/AutoRTFMTests/Private/UnorderedSet.cpp
--- a/Engine/Source/Programs/AutoRTFMTests/Private/UnorderedSet.cpp
+++ b/Engine/Source/Programs/AutoRTFMTests/Private/UnorderedSet.cpp
@@ -1,6 +1,7 @@
 // Copyright Epic Games, Inc. All Rights Reserved.
 
 #include "AutoRTFM.h"
+#include "AutoRTFM/Testing.h"
 #include "Catch2Includes.h"
 
 #include <unordered_set>
@@ -11,7 +12,7 @@ TEST_CASE("UnorderedSet")
 {
     constexpr unsigned Count = 100;
     unordered_set<unsigned> MySet;
-    AutoRTFM::Commit([&] ()
+    AutoRTFM::Testing::Commit([&] ()
     {
         unordered_set<unsigned> MySetInner;
         for (unsigned Index = Count; Index--;)
diff --git a/Engine/Source/Programs/UnrealBuildTool/Configuration/UEBuildModule.cs b/Engine/Source/Programs/UnrealBuildTool/Configuration/UEBuildModule.cs
--- a/Engine/Source/Programs/UnrealBuildTool/Configuration/UEBuildModule.cs
+++ b/Engine/Source/Programs/UnrealBuildTool/Configuration/UEBuildModule.cs
@@ -359,7 +359,7 @@ namespace UnrealBuildTool
 			// definition, like `[[my_attr]] void Foo();`, but these macros are used like `void FOO_API Foo();`, and if we used the C++ attribute name to
 			// disable AutoRTFM on this function the compiler would spit nails. But it'll happily consume the C-style attribute at that location and do
 			// the right thing, so we use that instead.
-			// Caution: Hard coded constant '1' in this string must be kept in sync with Engine\Source\Runtime\AutoRTFM\Public\AutoRTFMConstants.h.
+			// Caution: Hard coded constant '1' in this string must be kept in sync with Engine/Source/Runtime/AutoRTFM/Public/AutoRTFM/Constants.h.
 			Debug.Assert(Rules.Target.StaticAnalyzer == StaticAnalyzer.None || Rules.Target.StaticAnalyzer == StaticAnalyzer.Default);
 			bool bNonClangSA = Rules.Target.StaticAnalyzer == StaticAnalyzer.Default && Rules.Target.Platform.IsInGroup(UnrealPlatformGroup.Microsoft) && !Rules.Target.WindowsPlatform.Compiler.IsClang();
 			string AutoRTFMDisable = (Rules.Target.bUseAutoRTFMCompiler && Rules.bDisableAutoRTFMInstrumentation && Rules.bInjectAutoRTFMAttributeDisables && !bNonClangSA) ?
diff --git a/Engine/Source/Runtime/AutoRTFM/Private/API.cpp b/Engine/Source/Runtime/AutoRTFM/Private/API.cpp
--- a/Engine/Source/Runtime/AutoRTFM/Private/API.cpp
+++ b/Engine/Source/Runtime/AutoRTFM/Private/API.cpp
@@ -496,12 +496,6 @@ AUTORTFM_INTERNAL autortfm_result RTFM_autortfm_transact_then_open(void (*Uninst
 	return TransactThenOpenImpl(UninstrumentedWork, Arg, ReturnAddress);
 }
 
-AUTORTFM_INTERNAL void RTFM_autortfm_commit(void (*UninstrumentedWork)(void*), void (*InstrumentedWork)(void*), void* Arg)
-{
-	autortfm_result Result = autortfm_transact(UninstrumentedWork, InstrumentedWork, Arg);
-	AUTORTFM_FATAL_IF(Result != autortfm_committed, "Unexpected transaction result: %u", Result);
-}
-
 AUTORTFM_INTERNAL void RTFM_autortfm_abort()
 {
 	FContext* Context = FContext::Get();
@@ -870,17 +864,6 @@ extern "C" autortfm_result autortfm_transact_then_open(void (*UninstrumentedWork
 	return TransactThenOpenImpl(UninstrumentedWork, Arg, __builtin_return_address(0));
 }
 
-extern "C" void autortfm_commit(void (*UninstrumentedWork)(void*), void (*InstrumentedWork)(void*), void* Arg)
-{
-	if (autortfm_is_closed())
-	{
-		return RTFM_autortfm_commit(UninstrumentedWork, InstrumentedWork, Arg);
-	}
-
-    autortfm_result Result = autortfm_transact(UninstrumentedWork, InstrumentedWork, Arg);
-	AUTORTFM_FATAL_IF(Result != autortfm_committed, "Unexpected transaction result: %u", Result);
-}
-
 extern "C" void autortfm_abort() noexcept
 {
 	if (autortfm_is_closed())
diff --git a/Engine/Source/Runtime/AutoRTFM/Public/AutoRTFM.h b/Engine/Source/Runtime/AutoRTFM/Public/AutoRTFM.h
--- a/Engine/Source/Runtime/AutoRTFM/Public/AutoRTFM.h
+++ b/Engine/Source/Runtime/AutoRTFM/Public/AutoRTFM.h
@@ -272,22 +272,6 @@ UE_AUTORTFM_CRITICAL_INLINE ETransactionResult TransactThenOpen(const FunctorTyp
 	return Result;
 }
 
-// Run the callback in a transaction like Transact, but abort program
-// execution if the result is anything other than autortfm_committed.
-// Useful for testing.
-template<typename FunctorType>
-UE_AUTORTFM_CRITICAL_INLINE void Commit(AUTORTFM_IMPLICIT_ENABLE const FunctorType& Functor)
-{
-#if UE_AUTORTFM_ENABLED
-	static constexpr autortfm_mode FunctorAutoRTFMMode = CallMode<FunctorType>;
-	static_assert(FunctorAutoRTFMMode == autortfm_mode_enable, "AutoRTFM mode of functor passed to AutoRTFM::Commit() is not autortfm_mode_enable");
-#endif
-	autortfm_commit(
-		&AutoRTFMFunctorInvoker<FunctorType>,
-		AutoRTFMLookupInstrumentedFunctorInvoker<FunctorType>(Functor),
-		const_cast<void*>(static_cast<const void*>(&Functor)));
-}
-
 // Ends a transaction while in the closed, discarding all effects.
 // Sends control to the end of the transaction immediately.
 UE_AUTORTFM_CRITICAL_INLINE void AbortTransaction()
diff --git a/Engine/Source/Runtime/AutoRTFM/Public/AutoRTFM/CAPI.h b/Engine/Source/Runtime/AutoRTFM/Public/AutoRTFM/CAPI.h
--- a/Engine/Source/Runtime/AutoRTFM/Public/AutoRTFM/CAPI.h
+++ b/Engine/Source/Runtime/AutoRTFM/Public/AutoRTFM/CAPI.h
@@ -229,16 +229,6 @@ UE_AUTORTFM_CRITICAL_INLINE autortfm_result autortfm_transact_then_open(void (*U
 }
 #endif
 
-#if UE_AUTORTFM_ENABLED
-UE_AUTORTFM_API void autortfm_commit(void (*UninstrumentedWork)(void*), void (*InstrumentedWork)(void*), void* Arg);
-#else
-UE_AUTORTFM_CRITICAL_INLINE void autortfm_commit(void (*UninstrumentedWork)(void*), void (*InstrumentedWork)(void*), void* Arg)
-{
-	UE_AUTORTFM_UNUSED(InstrumentedWork);
-	UninstrumentedWork(Arg);
-}
-#endif
-
 #if UE_AUTORTFM_ENABLED
 UE_AUTORTFM_API void autortfm_abort_transaction() AUTORTFM_NOEXCEPT;
 #else
diff --git a/Engine/Source/Runtime/Engine/Private/Tests/AutoRTFM/AutoRTFMActorComponentTests.cpp b/Engine/Source/Runtime/Engine/Private/Tests/AutoRTFM/AutoRTFMActorComponentTests.cpp
--- a/Engine/Source/Runtime/Engine/Private/Tests/AutoRTFM/AutoRTFMActorComponentTests.cpp
+++ b/Engine/Source/Runtime/Engine/Private/Tests/AutoRTFM/AutoRTFMActorComponentTests.cpp
@@ -179,7 +179,7 @@ AUTORTFM_ACTOR_COMPONENT_TEST(RegisterComponent_UnregisterComponent)
 
 	bool bWasRegistered = false;
 
-	AutoRTFM::Commit([&]
+	AutoRTFM::Testing::Commit([&]
 		{
 			Component->RegisterComponent();
 			bWasRegistered = Component->IsRegistered();
@@ -197,7 +197,7 @@ AUTORTFM_ACTOR_COMPONENT_TEST(RegisterComponent_UnregisterComponent)
 	TestTrueExpr(AutoRTFM::ETransactionResult::AbortedByRequest == Result);
 	TestTrueExpr(Component->IsRegistered());
 
-	AutoRTFM::Commit([&]
+	AutoRTFM::Testing::Commit([&]
 		{
 			Component->UnregisterComponent();
 		});
@@ -223,7 +223,7 @@ AUTORTFM_ACTOR_COMPONENT_TEST(RegisterComponentWithWorld)
 	TestTrueExpr(AutoRTFM::ETransactionResult::AbortedByRequest == Result);
 	TestFalseExpr(Component->IsRegistered());
 
-	AutoRTFM::Commit([&]
+	AutoRTFM::Testing::Commit([&]
 		{
 			Component->RegisterComponentWithWorld(World);
 		});
@@ -263,7 +263,7 @@ AUTORTFM_ACTOR_COMPONENT_TEST(WeldTo)
 	TestTrueExpr(Component->IsWelded());
 	TestTrueExpr(&SomeInstance == Component->BodyInstance.WeldParent);
 
-	AutoRTFM::Commit([&]
+	AutoRTFM::Testing::Commit([&]
 		{
 			Component->WeldTo(Parent0);
 		});
@@ -280,7 +280,7 @@ AUTORTFM_ACTOR_COMPONENT_TEST(WeldTo)
 	TestTrueExpr(AutoRTFM::ETransactionResult::AbortedByRequest == Result);
 	TestFalseExpr(Component->IsWelded());
 
-	AutoRTFM::Commit([&]
+	AutoRTFM::Testing::Commit([&]
 		{
 			Component->WeldTo(Parent1);
 		});
@@ -296,7 +296,7 @@ AUTORTFM_ACTOR_COMPONENT_TEST(WeldTo)
 	TestTrueExpr(AutoRTFM::ETransactionResult::AbortedByRequest == Result);
 	TestFalseExpr(Component->IsWelded());
 
-	AutoRTFM::Commit([&]
+	AutoRTFM::Testing::Commit([&]
 		{
 			Component->UnWeldFromParent();
 		});
@@ -325,7 +325,7 @@ AUTORTFM_ACTOR_COMPONENT_TEST(FSparseDelegate)
 	TestTrueExpr(AutoRTFM::ETransactionResult::AbortedByRequest == Result);
 	TestFalseExpr(Object->bHitOnComponentPhysicsStateChanged);
 
-	AutoRTFM::Commit([&]
+	AutoRTFM::Testing::Commit([&]
 		{
 			Component->UnregisterComponent();
 		});
diff --git a/Engine/Source/Runtime/Engine/Private/Tests/AutoRTFM/AutoRTFMHttpModuleTests.cpp b/Engine/Source/Runtime/Engine/Private/Tests/AutoRTFM/AutoRTFMHttpModuleTests.cpp
--- a/Engine/Source/Runtime/Engine/Private/Tests/AutoRTFM/AutoRTFMHttpModuleTests.cpp
+++ b/Engine/Source/Runtime/Engine/Private/Tests/AutoRTFM/AutoRTFMHttpModuleTests.cpp
@@ -4,6 +4,7 @@
 #include "HttpManager.h"
 #include "HttpModule.h"
 #include "AutoRTFM.h"
+#include "AutoRTFM/Testing.h"
 #include "Containers/Array.h"
 #include "Interfaces/IHttpRequest.h"
 #include "Interfaces/IHttpResponse.h"
@@ -119,7 +120,7 @@ bool FAutoRTFMHttpModuleNetworking::RunTest(const FString& Parameters)
 		AutoRTFM::AbortTransaction();
 	});
 
-	AutoRTFM::Commit([&]
+	AutoRTFM::Testing::Commit([&]
 	{
 		TSharedRef<IHttpRequest> HttpRequest = Module.CreateRequest();
 		HttpRequest->SetVerb(TEXT("GET"));
@@ -137,7 +138,7 @@ bool FAutoRTFMHttpModuleNetworking::RunTest(const FString& Parameters)
 	TestTrue("HTTP requests are abandoned on abort", Module.NumActiveRequests() == 1);
 
 	// Verify that HttpModule can issue a request from within AutoRTFM transaction that is committed.
-	AutoRTFM::Commit([&]
+	AutoRTFM::Testing::Commit([&]
 	{
 		IssueHttpRequest(Module);
 		TestTrue("HTTP requests are deferred inside a transaction", Module.NumActiveRequests() == 1);
@@ -176,7 +177,7 @@ bool FAutoRTFMHttpRequestPreservesSettings::RunTest(const FString& Parameters)
 		// Test that changes are reflected both before and after a transactional commit.
 		{
 			TSharedPtr<IHttpRequest> Request;
-			AutoRTFM::Commit([&]
+			AutoRTFM::Testing::Commit([&]
 			{
 				Request = Module.CreateRequest();
 				SetterFn(*Request);  // Sets data on a FClosedHttpRequest
@@ -436,7 +437,7 @@ bool FAutoRTFMHttpRequestSafeToDelete::RunTest(const FString& Parameters)
 	FHttpTestModule Module;
 
 	// Verify that it is safe for an HTTP request to be deleted inside a transaction before it is committed.
-	AutoRTFM::Commit([&]
+	AutoRTFM::Testing::Commit([&]
 	{
 		TSharedPtr<IHttpRequest> HttpRequest = Module.CreateRequest();
 		HttpRequest->SetURL(TEXT("https://www.unrealengine.com/"));
@@ -455,7 +456,7 @@ bool FAutoRTFMHttpRequestSafeToDelete::RunTest(const FString& Parameters)
 	// Verify that it is safe for an HTTP request to be created inside a successful transaction, and then deleted
 	// in the open.
 	TSharedPtr<IHttpRequest> OuterRequest;
-	AutoRTFM::Commit([&]
+	AutoRTFM::Testing::Commit([&]
 	{
 		OuterRequest = Module.CreateRequest();
 		OuterRequest->SetURL(TEXT("https://www.unrealengine.com/"));
