Verse tests - Catch2 upgrade fix

#rnx
#Synced-CL 45819516

diff --git a/Engine/Source/Programs/AutoRTFMTests/Private/HitSetTests.cpp b/Engine/Source/Programs/AutoRTFMTests/Private/HitSetTests.cpp
--- a/Engine/Source/Programs/AutoRTFMTests/Private/HitSetTests.cpp
+++ b/Engine/Source/Programs/AutoRTFMTests/Private/HitSetTests.cpp
@@ -4,6 +4,7 @@
 #include "Catch2Includes.h"
 #include <cstdint>
 #include <unordered_set>
+#include <random>
 
 TEST_CASE("HitSet")
 {
@@ -61,7 +62,7 @@ TEST_CASE("HitSet")
 	auto RunTests = [&](auto Insert)
 	{
 		constexpr uintptr_t One = 1;
-		std::mt19937 Rand(0x1234);
+		Catch::SimplePcg32 Rand(0x1234);
 
 		SECTION("Vary Address")
 		{
