Change `RefCounting.h`'s to deprecate the return from `Release` and `GetRefCount`, and remove it entirely from `AddRef`.

#Synced-CL 48841801

diff --git a/Engine/Plugins/Experimental/UAF/UAF/Source/UAF/Public/UAF/Attributes/AttributeBindingData.h b/Engine/Plugins/Experimental/UAF/UAF/Source/UAF/Public/UAF/Attributes/AttributeBindingData.h
--- a/Engine/Plugins/Experimental/UAF/UAF/Source/UAF/Public/UAF/Attributes/AttributeBindingData.h
+++ b/Engine/Plugins/Experimental/UAF/UAF/Source/UAF/Public/UAF/Attributes/AttributeBindingData.h
@@ -94,7 +94,7 @@ namespace UE::UAF
 		// Returns a collection of the default attribute values for the specified named set
 		[[nodiscard]] UE_API FAttributeCollectionPtr GetDefaultAttributeValues(FName SetName) const;
 
-		// FRefCountedObject API
+		// UE::Private::FQueryableRefCountedObject API
 		UE_API uint32 Release() const;
 
 	private:
diff --git a/Engine/Source/Developer/ShaderCompilerCommon/Private/DXCWrapper.cpp b/Engine/Source/Developer/ShaderCompilerCommon/Private/DXCWrapper.cpp
--- a/Engine/Source/Developer/ShaderCompilerCommon/Private/DXCWrapper.cpp
+++ b/Engine/Source/Developer/ShaderCompilerCommon/Private/DXCWrapper.cpp
@@ -93,7 +93,7 @@ FDllHandle::~FDllHandle()
 
 FDxcModuleWrapper::FDxcModuleWrapper()
 {
-	if (GDxcHandle.GetRefCount() == 0)
+	if (!GDxcHandle.IsValid())
 	{
 		GDxilHandle = new FDllHandle(TEXT("dxil.dll"));
 		GDxcHandle = new FDllHandle(TEXT("dxcompiler.dll"));
@@ -123,7 +123,7 @@ FDxcModuleWrapper::~FDxcModuleWrapper()
 
 FShaderConductorModuleWrapper::FShaderConductorModuleWrapper()
 {
-	if (GShaderConductorHandle.GetRefCount() == 0)
+	if (!GShaderConductorHandle.IsValid())
 	{
 		GShaderConductorHandle = new FDllHandle(TEXT("ShaderConductor.dll"));
 	}
diff --git a/Engine/Source/Programs/AutoRTFMTests/Private/RefCountPtrTests.cpp b/Engine/Source/Programs/AutoRTFMTests/Private/RefCountPtrTests.cpp
--- a/Engine/Source/Programs/AutoRTFMTests/Private/RefCountPtrTests.cpp
+++ b/Engine/Source/Programs/AutoRTFMTests/Private/RefCountPtrTests.cpp
@@ -9,26 +9,90 @@ namespace
 {
 	struct FRefCountedInt : public FRefCountedObject
 	{
-		FRefCountedInt(int InValue) : Value(InValue) {}
+		explicit FRefCountedInt(int InValue, bool* InWasDestructed = nullptr) : Value(InValue), WasDestructed(InWasDestructed)
+		{
+			if (WasDestructed)
+			{
+				*WasDestructed = false;
+			}
+		}
+
 		int Value;
+		bool* WasDestructed;
+
+		~FRefCountedInt()
+		{
+			if (WasDestructed)
+			{
+				*WasDestructed = true;
+			}
+		}
 	};
 
 	struct FRefCountedMixinInt : public TRefCountingMixin<FRefCountedMixinInt>
 	{
-		FRefCountedMixinInt(int InValue) : Value(InValue) {}
+		explicit FRefCountedMixinInt(int InValue, bool* InWasDestructed = nullptr) : Value(InValue), WasDestructed(InWasDestructed)
+		{
+			if (WasDestructed)
+			{
+				*WasDestructed = false;
+			}
+		}
+
 		int Value;
+		bool* WasDestructed;
+
+		~FRefCountedMixinInt()
+		{
+			if (WasDestructed)
+			{
+				*WasDestructed = true;
+			}
+		}
 	};
 
 	struct FThreadSafeRefCountedInt : public FRefCountedObject
 	{
-		FThreadSafeRefCountedInt(int InValue) : Value(InValue) {}
+		explicit FThreadSafeRefCountedInt(int InValue, bool* InWasDestructed = nullptr) : Value(InValue), WasDestructed(InWasDestructed)
+		{
+			if (WasDestructed)
+			{
+				*WasDestructed = false;
+			}
+		}
+
 		int Value;
+		bool* WasDestructed;
+
+		~FThreadSafeRefCountedInt()
+		{
+			if (WasDestructed)
+			{
+				*WasDestructed = true;
+			}
+		}
 	};
 
 	struct FQueryableRefCountedInt : public UE::Private::FQueryableRefCountedObject
 	{
-		FQueryableRefCountedInt(int InValue) : Value(InValue) {}
+		explicit FQueryableRefCountedInt(int InValue, bool* InWasDestructed = nullptr) : Value(InValue), WasDestructed(InWasDestructed)
+		{
+			if (WasDestructed)
+			{
+				*WasDestructed = false;
+			}
+		}
+
 		int Value;
+		bool* WasDestructed;
+
+		~FQueryableRefCountedInt()
+		{
+			if (WasDestructed)
+			{
+				*WasDestructed = true;
+			}
+		}
 	};
 }
 
@@ -85,46 +149,46 @@ TEMPLATE_TEST_CASE("CheckRefCountsToDeprecate", "[RefCountPtr]", FRefCountedInt,
 	// Refcounts should be correct outside of a transaction.
 	SECTION("Non-transactional refcounts should be accurate")
 	{
-		TestType* Ptr = new TestType(42);
+		bool bWasDestructed = false;
 
-		REQUIRE(Ptr->GetRefCount() == 0);
-		(Ptr->AddRef().CheckAtLeast(1));
-		REQUIRE(Ptr->GetRefCount() == 1);
-		(Ptr->AddRef().CheckAtLeast(2));
-		REQUIRE(Ptr->GetRefCount() == 2);
-		REQUIRE(Ptr->Release() == 1);
-		REQUIRE(Ptr->GetRefCount() == 1);
-		(Ptr->AddRef().CheckAtLeast(2));
-		REQUIRE(Ptr->GetRefCount() == 2);
-		REQUIRE(Ptr->Release() == 1);
-		REQUIRE(Ptr->GetRefCount() == 1);
-		REQUIRE(Ptr->Release() == 0);
+		{
+			TestType* Ptr = new TestType(42, &bWasDestructed);
 
-		// Ptr is dead
+			Ptr->AddRef();
+			Ptr->AddRef();
+			Ptr->Release();
+			Ptr->AddRef();
+			Ptr->Release();
+
+			REQUIRE(!bWasDestructed);
+			Ptr->Release();
+
+			// We aren't in a transaction so we will have destructed the reference already!
+			REQUIRE(bWasDestructed);
+		}
+
+		REQUIRE(bWasDestructed);
 	}
 
 	// Refcounts can be inflated inside a transaction, since releases are deferred.
 	SECTION("Transactional refcounts should represent a lower bound")
 	{
-		AutoRTFM::Testing::Commit([]
-			{
-				TestType* Ptr = new TestType(42);
-
-				REQUIRE(Ptr->GetRefCount() >= 0);
-				(Ptr->AddRef().CheckAtLeast(1));
-				REQUIRE(Ptr->GetRefCount() >= 1);
-				(Ptr->AddRef().CheckAtLeast(2));
-				REQUIRE(Ptr->GetRefCount() >= 2);
-				REQUIRE(Ptr->Release() >= 1);
-				REQUIRE(Ptr->GetRefCount() >= 1);
-				(Ptr->AddRef().CheckAtLeast(2));
-				REQUIRE(Ptr->GetRefCount() >= 2);
-				REQUIRE(Ptr->Release() >= 1);
-				REQUIRE(Ptr->GetRefCount() >= 1);
-				REQUIRE(Ptr->Release() >= 0);
-
-				// Ptr is dead
-			});
+		bool bWasDestructed = false;
+
+		AutoRTFM::Testing::Commit([&]
+		{
+			TestType* Ptr = new TestType(42, &bWasDestructed);
+
+			Ptr->AddRef();
+			Ptr->AddRef();
+			Ptr->Release();
+			Ptr->AddRef();
+			Ptr->Release();
+			Ptr->Release();
+			REQUIRE(!bWasDestructed);
+		});
+
+		REQUIRE(bWasDestructed);
 	}
 }
 
@@ -132,32 +196,44 @@ TEMPLATE_TEST_CASE("PreviouslyAllocated", "[RefCountPtr]", FRefCountedInt, FRefC
 {
 	SECTION("Using new T()")
 	{
-		TRefCountPtr<TestType> Foo(new TestType(42));
+		bool bWasDestructed = false;
 
-		AutoRTFM::Testing::Commit([&]
 		{
-			// Make a copy to bump the reference count.
-			TRefCountPtr<TestType> Copy = Foo;
-			Copy->Value = 13;
-		});
+			TRefCountPtr<TestType> Foo(new TestType(42, &bWasDestructed));
 
-		REQUIRE(Foo->Value == 13);
-		REQUIRE(Foo.GetRefCount() == 1);
+			AutoRTFM::Testing::Commit([&]
+			{
+				// Make a copy to bump the reference count.
+				TRefCountPtr<TestType> Copy = Foo;
+				Copy->Value = 13;
+			});
+
+			REQUIRE(Foo->Value == 13);
+			REQUIRE(!bWasDestructed);
+		}
+
+		REQUIRE(bWasDestructed);
 	}
 
 	SECTION("Using MakeRefCount<T>")
 	{
-		TRefCountPtr<TestType> Foo = MakeRefCount<TestType>(42);
+		bool bWasDestructed = false;
 
-		AutoRTFM::Testing::Commit([&]
 		{
-			// Make a copy to bump the reference count.
-			TRefCountPtr<TestType> Copy = Foo;
-			Copy->Value = 13;
-		});
+			TRefCountPtr<TestType> Foo = MakeRefCount<TestType>(42, &bWasDestructed);
 
-		REQUIRE(Foo->Value == 13);
-		REQUIRE(Foo.GetRefCount() == 1);
+			AutoRTFM::Testing::Commit([&]
+			{
+				// Make a copy to bump the reference count.
+				TRefCountPtr<TestType> Copy = Foo;
+				Copy->Value = 13;
+			});
+
+			REQUIRE(Foo->Value == 13);
+			REQUIRE(!bWasDestructed);
+		}
+
+		REQUIRE(bWasDestructed);
 	}
 }
 
@@ -165,36 +241,48 @@ TEMPLATE_TEST_CASE("AbortWithPreviouslyAllocated", "[RefCountPtr]", FRefCountedI
 {
 	SECTION("Using new T()")
 	{
-		TRefCountPtr<TestType> Foo(new TestType(42));
+		bool bWasDestructed = false;
 
-		AutoRTFM::Testing::Abort([&]
 		{
-			// Make a copy to bump the reference count.
-			TRefCountPtr<TestType> Copy = Foo;
-			Copy->Value = 13;
+			TRefCountPtr<TestType> Foo(new TestType(42, &bWasDestructed));
 
-			AutoRTFM::AbortTransaction();
-		});
+			AutoRTFM::Testing::Abort([&]
+			{
+				// Make a copy to bump the reference count.
+				TRefCountPtr<TestType> Copy = Foo;
+				Copy->Value = 13;
 
-		REQUIRE(Foo->Value == 42);
-		REQUIRE(Foo.GetRefCount() == 1);
+				AutoRTFM::AbortTransaction();
+			});
+
+			REQUIRE(Foo->Value == 42);
+			REQUIRE(!bWasDestructed);
+		}
+
+		REQUIRE(bWasDestructed);
 	}
 
 	SECTION("Using MakeRefCount<T>")
 	{
-		TRefCountPtr<TestType> Foo = MakeRefCount<TestType>(42);
+		bool bWasDestructed = false;
 
-		AutoRTFM::Testing::Abort([&]
 		{
-			// Make a copy to bump the reference count.
-			TRefCountPtr<TestType> Copy = Foo;
-			Copy->Value = 13;
+			TRefCountPtr<TestType> Foo = MakeRefCount<TestType>(42, &bWasDestructed);
 
-			AutoRTFM::AbortTransaction();
-		});
+			AutoRTFM::Testing::Abort([&]
+			{
+				// Make a copy to bump the reference count.
+				TRefCountPtr<TestType> Copy = Foo;
+				Copy->Value = 13;
 
-		REQUIRE(Foo->Value == 42);
-		REQUIRE(Foo.GetRefCount() == 1);
+				AutoRTFM::AbortTransaction();
+			});
+
+			REQUIRE(Foo->Value == 42);
+			REQUIRE(!bWasDestructed);
+		}
+
+		REQUIRE(bWasDestructed);
 	}
 }
 
diff --git a/Engine/Source/Runtime/Core/Private/Experimental/IO/IoStatusError.cpp b/Engine/Source/Runtime/Core/Private/Experimental/IO/IoStatusError.cpp
--- a/Engine/Source/Runtime/Core/Private/Experimental/IO/IoStatusError.cpp
+++ b/Engine/Source/Runtime/Core/Private/Experimental/IO/IoStatusError.cpp
@@ -30,9 +30,9 @@ namespace UE::UnifiedError::IoStore
 			ErrorMessage = InErrorMessage;
 		}
 
-		virtual FReturnedRefCountValue AddRef() const final override { return FRefCountedObject::AddRef(); }
-		virtual uint32 Release() const final override { return FRefCountedObject::Release(); }
-		virtual uint32 GetRefCount() const final override { return FRefCountedObject::GetRefCount(); }
+		virtual void AddRef() const final override { return FRefCountedObject::AddRef(); }
+		virtual FReturnedRefCountValue Release() const final override { return FRefCountedObject::Release(); }
+		virtual FReturnedRefCountValue GetRefCount() const final override { return FRefCountedObject::GetRefCount(); }
 
 		const FUtf8String& GetErrorName(const FError& Error) const
 		{
diff --git a/Engine/Source/Runtime/Core/Private/Experimental/Misc/ExecutionResource.cpp b/Engine/Source/Runtime/Core/Private/Experimental/Misc/ExecutionResource.cpp
--- a/Engine/Source/Runtime/Core/Private/Experimental/Misc/ExecutionResource.cpp
+++ b/Engine/Source/Runtime/Core/Private/Experimental/Misc/ExecutionResource.cpp
@@ -93,17 +93,17 @@ public:
 	{
 	}
 
-	FReturnedRefCountValue AddRef() const override
+	void AddRef() const override
 	{
 		return FRefCountedObject::AddRef();
 	}
 
-	uint32 Release() const override
+	FReturnedRefCountValue Release() const override
 	{
 		return FRefCountedObject::Release();
 	}
 	
-	uint32 GetRefCount() const override
+	FReturnedRefCountValue GetRefCount() const override
 	{
 		return FRefCountedObject::GetRefCount();
 	}
diff --git a/Engine/Source/Runtime/Core/Private/Internationalization/TextHistory.h b/Engine/Source/Runtime/Core/Private/Internationalization/TextHistory.h
--- a/Engine/Source/Runtime/Core/Private/Internationalization/TextHistory.h
+++ b/Engine/Source/Runtime/Core/Private/Internationalization/TextHistory.h
@@ -149,9 +149,9 @@ public:
 	FTextHistory& operator=(FTextHistory&) = delete;
 
 	//~ IRefCountedObject
-	virtual FReturnedRefCountValue AddRef() const override final { return TRefCountingMixin<FTextHistory>::AddRef(); }
-	virtual uint32 Release() const override final { return TRefCountingMixin<FTextHistory>::Release(); }
-	virtual uint32 GetRefCount() const override final { return TRefCountingMixin<FTextHistory>::GetRefCount(); }
+	virtual void AddRef() const override final { TRefCountingMixin<FTextHistory>::AddRef(); }
+	virtual FReturnedRefCountValue Release() const override final { return TRefCountingMixin<FTextHistory>::Release(); }
+	virtual FReturnedRefCountValue GetRefCount() const override final { return TRefCountingMixin<FTextHistory>::GetRefCount(); }
 
 	//~ ITextData
 	virtual const FString& GetSourceString() const override { return GetDisplayString(); }
diff --git a/Engine/Source/Runtime/Core/Private/Misc/QueuedThreadPoolWrapper.cpp b/Engine/Source/Runtime/Core/Private/Misc/QueuedThreadPoolWrapper.cpp
--- a/Engine/Source/Runtime/Core/Private/Misc/QueuedThreadPoolWrapper.cpp
+++ b/Engine/Source/Runtime/Core/Private/Misc/QueuedThreadPoolWrapper.cpp
@@ -151,7 +151,7 @@ bool FQueuedThreadPoolWrapper::RetractQueuedWork(IQueuedWork* InQueuedWork)
 	if (Retracted)
 	{
 		// When we retract a task, there should be no external refs on it
-		check(Retracted->GetRefCount() == 1);
+		check(Retracted->NumRefs.GetValue() == 1);
 		Retracted->Release();
 	}
 	return Retracted != nullptr;
@@ -164,13 +164,13 @@ int32 FQueuedThreadPoolWrapper::GetNumThreads() const
 
 void FQueuedThreadPoolWrapper::ReleaseWorkNoLock(FScheduledWork* Work)
 {
-	check(Work->GetRefCount() == 0);
+	check(Work->NumRefs.GetValue() == 0);
 	CurrentConcurrency--;
 	OnUnscheduled(Work);
 	
 	ScheduledWork.Remove(Work->GetInnerWork());
 	Work->Reset();
-	check(Work->GetRefCount() == 0);
+	check(Work->NumRefs.GetValue() == 0);
 	WorkPool.Push(Work);
 }
 
@@ -207,7 +207,7 @@ bool FQueuedThreadPoolWrapper::TryRetractWorkNoLock(EQueuedWorkPriority InPriori
 			if (WrappedQueuedThreadPool->RetractQueuedWork(Pair.Value))
 			{
 				// When we retract a task, there should be no external refs on it
-				check(Pair.Value->GetRefCount() == 1);
+				check(Pair.Value->NumRefs.GetValue() == 1);
 				QueuedWork.Enqueue(Pair.Key, Pair.Value->GetPriority());
 				Pair.Value->Release();
 
diff --git a/Engine/Source/Runtime/Core/Public/Experimental/UnifiedError/UnifiedError.h b/Engine/Source/Runtime/Core/Public/Experimental/UnifiedError/UnifiedError.h
--- a/Engine/Source/Runtime/Core/Public/Experimental/UnifiedError/UnifiedError.h
+++ b/Engine/Source/Runtime/Core/Public/Experimental/UnifiedError/UnifiedError.h
@@ -276,7 +276,6 @@ namespace UE::UnifiedError
 	class FManditoryErrorDetails : public IErrorDetails
 	{
 		UE_DECLARE_FERROR_DETAILS_ABSTRACT(UnifiedError, FManditoryErrorDetails);
-	private:
 	public:
 		virtual FAnsiString GetErrorCodeString(const FError& Error) const = 0;
 		virtual FAnsiString GetModuleIdString(const FError& Error) const = 0;
@@ -306,12 +305,11 @@ namespace UE::UnifiedError
 	{
 		UE_DECLARE_FERROR_DETAILS_ABSTRACT(UnifiedError, FRefCountedErrorDetails);
 
-
 		CORE_API virtual ~FRefCountedErrorDetails();
 
-		virtual FReturnedRefCountValue AddRef() const final override { return FRefCountedObject::AddRef(); }
-		virtual uint32 Release() const final override { return FRefCountedObject::Release(); }
-		virtual uint32 GetRefCount() const final override { return FRefCountedObject::GetRefCount(); }
+		virtual void AddRef() const final override { FRefCountedObject::AddRef(); }
+		virtual FReturnedRefCountValue Release() const final override { return FRefCountedObject::Release(); }
+		virtual FReturnedRefCountValue GetRefCount() const final override { return FRefCountedObject::GetRefCount(); }
 	};
 
 
@@ -608,11 +606,10 @@ namespace UE::UnifiedError
 
 		// IRefCountedObject implementation
 		//  FStaticErrorDetails is statically allocated, make sure it is never released.
-		virtual FReturnedRefCountValue AddRef() const final { return FReturnedRefCountValue(10); }
-		virtual uint32 Release() const final { return 10; }
-		virtual uint32 GetRefCount() const final { return 10; }
+		virtual void AddRef() const final {}
+		virtual FReturnedRefCountValue Release() const final { return FReturnedRefCountValue{ 10 }; }
+		virtual FReturnedRefCountValue GetRefCount() const final { return FReturnedRefCountValue{ 10 }; }
 
-		
 		virtual bool ShouldInclude(const EDetailFilter Filter) const override { return EnumHasAnyFlags(Filter, EDetailFilter::Default); }
 	};
 
diff --git a/Engine/Source/Runtime/Core/Public/Misc/QueuedThreadPoolWrapper.h b/Engine/Source/Runtime/Core/Public/Misc/QueuedThreadPoolWrapper.h
--- a/Engine/Source/Runtime/Core/Public/Misc/QueuedThreadPoolWrapper.h
+++ b/Engine/Source/Runtime/Core/Public/Misc/QueuedThreadPoolWrapper.h
@@ -113,12 +113,12 @@ protected:
 			Release();
 		}
 	private:
-		FReturnedRefCountValue AddRef() const override
+		void AddRef() const override
 		{
-			return FReturnedRefCountValue{uint32(NumRefs.Increment())};
+			NumRefs.Increment();
 		}
 
-		uint32 Release() const override
+		FReturnedRefCountValue Release() const override
 		{
 			uint32 Refs = uint32(NumRefs.Decrement());
 
@@ -129,25 +129,24 @@ protected:
 			{
 				ParentPool->Schedule(const_cast<FScheduledWork*>(this));
 			}
-			return Refs;
+
+			return FReturnedRefCountValue{ Refs };
 		}
 
-		uint32 GetRefCount() const override
+		FReturnedRefCountValue GetRefCount() const override
 		{
-			return uint32(NumRefs.GetValue());
+			return FReturnedRefCountValue{ static_cast<uint32>(NumRefs.GetValue()) };
 		}
 
 		void Assign(FQueuedThreadPoolWrapper* InParentPool, IQueuedWork* InWork, EQueuedWorkPriority InPriority)
 		{
-			check(GetRefCount() == 0);
+			check(NumRefs.GetValue() == 0);
 			ParentPool = InParentPool;
 			Work = InWork;
 			Priority = InPriority;
 			AddRef();
 		}
 
-		
-
 		void Abandon() override
 		{
 			Work->Abandon();
diff --git a/Engine/Source/Runtime/Core/Public/Templates/RefCounting.h b/Engine/Source/Runtime/Core/Public/Templates/RefCounting.h
--- a/Engine/Source/Runtime/Core/Public/Templates/RefCounting.h
+++ b/Engine/Source/Runtime/Core/Public/Templates/RefCounting.h
@@ -36,7 +36,7 @@ struct FReturnedRefCountValue
 	FReturnedRefCountValue& operator=(const FReturnedRefCountValue& Other) = default;
 	FReturnedRefCountValue& operator=(FReturnedRefCountValue&& Other) = default;
 
-	UE_DEPRECATED(5.6, "Inspecting an object's refcount is deprecated.")
+	UE_DEPRECATED(5.8, "Inspecting an object's refcount is deprecated.")
 	operator uint32() const
 	{
 		return RefCount;
@@ -293,14 +293,11 @@ class IRefCountedObject
 {
 public:
 	virtual ~IRefCountedObject() = default;
-	virtual FReturnedRefCountValue AddRef() const = 0;
+	virtual void AddRef() const = 0;
 
-	// TODO (SOL-7350): return FReturnedRefCountValue from Release(); clean up call sites
-	// which rely on its return value.
-	virtual uint32 Release() const = 0;
+	virtual FReturnedRefCountValue Release() const = 0;
 
-	// TODO (SOL-7350): mark this function as deprecated; clean up existing callers.
-	virtual uint32 GetRefCount() const = 0;
+	virtual FReturnedRefCountValue GetRefCount() const = 0;
 };
 
 /**
@@ -315,19 +312,19 @@ public:
 	FRefCountedObject(const FRefCountedObject& Rhs) = delete;
 	FRefCountedObject& operator=(const FRefCountedObject& Rhs) = delete;
 
-	FReturnedRefCountValue AddRef() const
+	void AddRef() const
 	{
-		return FReturnedRefCountValue{Super::AddRef<DeleteThis>()};
+		Super::AddRef<DeleteThis>();
 	}
 
-	uint32 Release() const
+	FReturnedRefCountValue Release() const
 	{
-		return Super::Release<DeleteThis>();
+		return FReturnedRefCountValue{ Super::Release<DeleteThis>() };
 	}
 
-	uint32 GetRefCount() const
+	FReturnedRefCountValue GetRefCount() const
 	{
-		return Super::GetRefCount();
+		return FReturnedRefCountValue{ Super::GetRefCount() };
 	}
 
 private:
@@ -343,8 +340,6 @@ class UE_DEPRECATED(5.8, "`FRefCountedObject` is now thread-safe so use that ins
 
 class UE_DEPRECATED(5.8, "`FRefCountedObject` is now thread-safe so use that instead.") FThreadSafeRefCountedObject : public FRefCountedObject {};
 
-
-
 /**
  * ERefCountingMode is used select between either 'fast' or 'thread safe' ref-counting types.
  * This is only used at compile time to select between template specializations.
@@ -395,19 +390,19 @@ public:
 	TRefCountingMixin(const TRefCountingMixin&) = delete;
 	TRefCountingMixin& operator=(const TRefCountingMixin&) = delete;
 
-	FReturnedRefCountValue AddRef() const
+	void AddRef() const
 	{
-		return FReturnedRefCountValue{Super::template AddRef<StaticDestroyMixin>()};
+		Super::template AddRef<StaticDestroyMixin>();
 	}
 
-	uint32 Release() const
+	FReturnedRefCountValue Release() const
 	{
-		return Super::template Release<StaticDestroyMixin>();
+		return FReturnedRefCountValue{ Super::template Release<StaticDestroyMixin>() };
 	}
 
-	uint32 GetRefCount() const
+	FReturnedRefCountValue GetRefCount() const
 	{
-		return Super::GetRefCount();
+		return FReturnedRefCountValue{ Super::GetRefCount() };
 	}
 
 	static void StaticDestroyObject(const T* Obj)
@@ -438,12 +433,12 @@ public:
 	TRefCountingMixin(const TRefCountingMixin&) = delete;
 	TRefCountingMixin& operator=(const TRefCountingMixin&) = delete;
 
-	FReturnedRefCountValue AddRef() const
+	void AddRef() const
 	{
-		return FReturnedRefCountValue{++RefCount};
+		++RefCount;
 	}
 
-	uint32 Release() const
+	FReturnedRefCountValue Release() const
 	{
 		checkSlow(RefCount > 0);
 
@@ -453,12 +448,12 @@ public:
 		}
 
 		// Note: TRefCountPtr doesn't use the return value
-		return 0;
+		return FReturnedRefCountValue{ 0 };
 	}
 
-	uint32 GetRefCount() const
+	FReturnedRefCountValue GetRefCount() const
 	{
-		return RefCount;
+		return FReturnedRefCountValue{ RefCount };
 	}
 
 	static void StaticDestroyObject(const T* Obj)
diff --git a/Engine/Source/Runtime/D3D12RHI/Private/D3D12Resources.h b/Engine/Source/Runtime/D3D12RHI/Private/D3D12Resources.h
--- a/Engine/Source/Runtime/D3D12RHI/Private/D3D12Resources.h
+++ b/Engine/Source/Runtime/D3D12RHI/Private/D3D12Resources.h
@@ -1036,17 +1036,17 @@ public:
 	void ReleaseOwnership();
 
 	// IRefCountedObject interface.
-	virtual FReturnedRefCountValue AddRef() const
+	virtual void AddRef() const override final
 	{
-		return FReturnedRefCountValue{FRHIResource::AddRef()};
+		FRHIResource::AddRef();
 	}
-	virtual uint32 Release() const
+	virtual FReturnedRefCountValue Release() const override final
 	{
-		return FRHIResource::Release();
+		return FReturnedRefCountValue{ FRHIResource::Release() };
 	}
-	virtual uint32 GetRefCount() const
+	virtual FReturnedRefCountValue GetRefCount() const override final
 	{
-		return FRHIResource::GetRefCount();
+		return FReturnedRefCountValue{ FRHIResource::GetRefCount() };
 	}
 
 	static void GetResourceDescAndAlignment(const FRHIBufferCreateDesc& CreateDesc, D3D12_RESOURCE_DESC& ResourceDesc, uint32& Alignment);
diff --git a/Engine/Source/Runtime/D3D12RHI/Private/D3D12Texture.h b/Engine/Source/Runtime/D3D12RHI/Private/D3D12Texture.h
--- a/Engine/Source/Runtime/D3D12RHI/Private/D3D12Texture.h
+++ b/Engine/Source/Runtime/D3D12RHI/Private/D3D12Texture.h
@@ -45,9 +45,9 @@ public:
 	virtual ~FD3D12Texture();
 
 	// IRefCountedObject interface overrides from FD3D12BaseShaderResource
-	virtual FReturnedRefCountValue AddRef() const override		{ return FReturnedRefCountValue{FRHIResource::AddRef()};	}
-	virtual uint32 Release() const override						{ return FRHIResource::Release();	}
-	virtual uint32 GetRefCount() const override					{ return FRHIResource::GetRefCount(); }
+	virtual void AddRef() const override                        { FRHIResource::AddRef(); }
+	virtual FReturnedRefCountValue Release() const override	    { return FReturnedRefCountValue{ FRHIResource::Release() }; }
+	virtual FReturnedRefCountValue GetRefCount() const override { return  FReturnedRefCountValue{ FRHIResource::GetRefCount() }; }
 
 	// FRHIResource overrides 
 #if RHI_ENABLE_RESOURCE_INFO
diff --git a/Engine/Source/Runtime/D3D12RHI/Private/D3D12Viewport.cpp b/Engine/Source/Runtime/D3D12RHI/Private/D3D12Viewport.cpp
--- a/Engine/Source/Runtime/D3D12RHI/Private/D3D12Viewport.cpp
+++ b/Engine/Source/Runtime/D3D12RHI/Private/D3D12Viewport.cpp
@@ -283,12 +283,6 @@ bool ReleaseBackBufferResource(TRefCountPtr<TResource>& Resource, const TCHAR* E
 	if (bValidReference)
 	{
 		// Tell the back buffer to delete immediately so that we can call resize.
-		if (Resource->GetRefCount() != 1)
-		{
-			UE_LOG(LogD3D12RHI, Log, TEXT("%s %d leaking with %d refs during Resize."), ErrorName, Index, Resource->GetRefCount());
-		}
-		check(Resource->GetRefCount() == 1);
-
 		for (TSubresource& SubResource : *Resource)
 		{
 			SubResource.GetResource()->DoNotDeferDelete();
diff --git a/Engine/Source/Runtime/Engine/Private/Streaming/RenderAssetUpdate.cpp b/Engine/Source/Runtime/Engine/Private/Streaming/RenderAssetUpdate.cpp
--- a/Engine/Source/Runtime/Engine/Private/Streaming/RenderAssetUpdate.cpp
+++ b/Engine/Source/Runtime/Engine/Private/Streaming/RenderAssetUpdate.cpp
@@ -177,7 +177,7 @@ FRenderAssetUpdate::~FRenderAssetUpdate()
 	check(TaskSynchronization.GetValue() == 0 && TaskState == TS_Done);
 }
 
-uint32 FRenderAssetUpdate::Release() const
+FReturnedRefCountValue FRenderAssetUpdate::Release() const
 {
 	uint32 NewValue = (uint32)NumRefs.Decrement();
 	if (NewValue == 0)
@@ -192,7 +192,7 @@ uint32 FRenderAssetUpdate::Release() const
 			UE_LOG(LogContentStreaming, Error, TEXT("RenderAssetUpdate is leaking (State=%d)"), (int32)TaskState);
 		}
 	}
-	return NewValue;
+	return FReturnedRefCountValue{NewValue};
 }
 
 void FRenderAssetUpdate::Tick(EThreadType InCurrentThread)
diff --git a/Engine/Source/Runtime/Engine/Public/RenderAssetUpdate.h b/Engine/Source/Runtime/Engine/Public/RenderAssetUpdate.h
--- a/Engine/Source/Runtime/Engine/Public/RenderAssetUpdate.h
+++ b/Engine/Source/Runtime/Engine/Public/RenderAssetUpdate.h
@@ -130,16 +130,16 @@ public:
 
 	//****************** IRefCountedObject ****************/
 
-	FReturnedRefCountValue AddRef() const final override 
+	void AddRef() const final override 
 	{ 
-		return FReturnedRefCountValue{(uint32)NumRefs.Increment()};
+		NumRefs.Increment();
 	}
 
-	UE_API uint32 Release() const final override;
+	UE_API FReturnedRefCountValue Release() const final override;
 
-	uint32 GetRefCount() const final override
+	FReturnedRefCountValue GetRefCount() const final override
 	{
-		return (uint32)NumRefs.GetValue();
+		return FReturnedRefCountValue{ static_cast<uint32>(NumRefs.GetValue()) };
 	}
 
 protected:
diff --git a/Engine/Source/Runtime/Renderer/Private/SceneRendering.cpp b/Engine/Source/Runtime/Renderer/Private/SceneRendering.cpp
--- a/Engine/Source/Runtime/Renderer/Private/SceneRendering.cpp
+++ b/Engine/Source/Runtime/Renderer/Private/SceneRendering.cpp
@@ -3622,7 +3622,7 @@ struct FCrossGPUTarget
 	TArray<FCrossGPUTransfer> Transfers;
 };
 
-class FCrossGPUTransfersDeferred : public FRefCountedObject
+class FCrossGPUTransfersDeferred : public UE::Private::FQueryableRefCountedObject
 {
 public:
 	TArray<FCrossGPUTarget> Targets;
diff --git a/Engine/Source/Runtime/VulkanRHI/Private/VulkanResources.h b/Engine/Source/Runtime/VulkanRHI/Private/VulkanResources.h
--- a/Engine/Source/Runtime/VulkanRHI/Private/VulkanResources.h
+++ b/Engine/Source/Runtime/VulkanRHI/Private/VulkanResources.h
@@ -271,17 +271,17 @@ public:
 	enum { StaticFrequency = ShaderType };
 
 	// IRefCountedObject interface.
-	virtual FReturnedRefCountValue AddRef() const override final
+	virtual void AddRef() const override final
 	{
-		return FReturnedRefCountValue{FRHIResource::AddRef()};
+		FRHIResource::AddRef();
 	}
-	virtual uint32 Release() const override final
+	virtual FReturnedRefCountValue Release() const override final
 	{
-		return FRHIResource::Release();
+		return FReturnedRefCountValue{ FRHIResource::Release() };
 	}
-	virtual uint32 GetRefCount() const override final
+	virtual FReturnedRefCountValue GetRefCount() const override final
 	{
-		return FRHIResource::GetRefCount();
+		return FReturnedRefCountValue{ FRHIResource::GetRefCount() };
 	}
 };
 
@@ -316,17 +316,17 @@ public:
 	TRefCountPtr<FVulkanShaderModule> GetOrCreateHandle(uint32 ModuleIdentifier);
 
 	// IRefCountedObject interface.
-	virtual FReturnedRefCountValue AddRef() const override final
+	virtual void AddRef() const override final
 	{
-		return FReturnedRefCountValue{FRHIResource::AddRef()};
+		FRHIResource::AddRef();
 	}
-	virtual uint32 Release() const override final
+	virtual FReturnedRefCountValue Release() const override final
 	{
-		return FRHIResource::Release();
+		return FReturnedRefCountValue{ FRHIResource::Release() };
 	}
-	virtual uint32 GetRefCount() const override final
+	virtual FReturnedRefCountValue GetRefCount() const override final
 	{
-		return FRHIResource::GetRefCount();
+		return FReturnedRefCountValue{ FRHIResource::GetRefCount() };
 	}
 };
 
