[AutoRTFM] A handful of test fixes for the thorough AutoRTFM sanitizer

#Synced-CL 48740692

diff --git a/Engine/Source/Programs/AutoRTFMTests/Private/API.cpp b/Engine/Source/Programs/AutoRTFMTests/Private/API.cpp
--- a/Engine/Source/Programs/AutoRTFMTests/Private/API.cpp
+++ b/Engine/Source/Programs/AutoRTFMTests/Private/API.cpp
@@ -1276,6 +1276,8 @@ TEST_CASE("API.OnAbortMacro_NoAbort")
 
 TEST_CASE("API.OnPreAbortMacro_NoAbort")
 {
+	AUTORTFM_SCOPED_DISABLE_MEMORY_VALIDATION();
+
 	int Value = 123;
 
 	UE_AUTORTFM_TRANSACT
@@ -1312,6 +1314,8 @@ TEST_CASE("API.OnAbortMacro_WithAbort")
 
 TEST_CASE("API.OnPreAbortMacro_WithAbort")
 {
+	AUTORTFM_SCOPED_DISABLE_MEMORY_VALIDATION();
+
 	int Value = 123;
 
 	UE_AUTORTFM_TRANSACT
diff --git a/Engine/Source/Programs/AutoRTFMTests/Private/OpenAPI.cpp b/Engine/Source/Programs/AutoRTFMTests/Private/OpenAPI.cpp
--- a/Engine/Source/Programs/AutoRTFMTests/Private/OpenAPI.cpp
+++ b/Engine/Source/Programs/AutoRTFMTests/Private/OpenAPI.cpp
@@ -83,8 +83,7 @@ TEST_CASE("OpenAPI.WriteDataInTheOpen")
 	{
 		AutoRTFM::Open([&]()
 		{
-			AutoRTFM::RecordOpenWrite(&Value);
-			Value = 1;
+			AutoRTFM::Assign(Value, 1);
 		});
 	});
 
@@ -553,8 +552,7 @@ TEST_CASE("OpenAPI.StackWriteCommitInTheOpen1")
 		AutoRTFM::Open([&]()
 		{
 			AutoRTFM::ForTheRuntime::StartTransaction();
-			AutoRTFM::RecordOpenWrite(&Value);
-			Value = 10;
+			AutoRTFM::Assign(Value, 10);
 			AutoRTFM::ForTheRuntime::CommitTransaction();
 			REQUIRE(Value == 10);
 		});
@@ -575,8 +573,7 @@ TEST_CASE("OpenAPI.StackWriteCommitInTheOpen2")
 			{
 				AutoRTFM::Open([&]()
 				{
-					AutoRTFM::RecordOpenWrite(&Value);
-					Value = 10;
+					AutoRTFM::Assign(Value, 10);
 				});
 			});
 			REQUIRE(AutoRTFM::EContextStatus::OnTrack == Status);
@@ -596,8 +593,7 @@ TEST_CASE("OpenAPI.StackWriteAbortInTheOpen1")
 		AutoRTFM::Open([&]()
 		{
 			AutoRTFM::ForTheRuntime::StartTransaction();
-			AutoRTFM::RecordOpenWrite(&Value);
-			Value = 10;
+			AutoRTFM::Assign(Value, 10);
 			AutoRTFM::ForTheRuntime::RollbackTransaction();
 			REQUIRE(Value == 0);
 		});
@@ -626,15 +622,13 @@ int Value1 = 0;
 
 TEST_CASE("OpenAPI.WriteMemory1")
 {
-	const int sourceValue = 10;
 	AutoRTFM::Transact([&]()
 	{
 		AutoRTFM::Open([&]()
 		{
 			AutoRTFM::ForTheRuntime::StartTransaction();
 			
-			AutoRTFM::RecordOpenWrite(&Value1);
-			Value1 = sourceValue;
+			AutoRTFM::Assign(Value1, 10);
 
 			AutoRTFM::ForTheRuntime::CommitTransaction();
 		});
@@ -652,9 +646,7 @@ TEST_CASE("OpenAPI.StackWriteAbortInTheOpen2")
 		{
 			AutoRTFM::ForTheRuntime::StartTransaction();
 
-			// Illegal to write to Value because it's in the inner-most closed-nest
-			AutoRTFM::RecordOpenWrite(&Value);
-			Value = 10;
+			AutoRTFM::Assign(Value, 10);
 
 			AutoRTFM::ForTheRuntime::RollbackTransaction();
 		});
@@ -688,8 +680,7 @@ TEST_CASE("OpenAPI.WriteTrivialStructure")
 		{
 			AutoRTFM::ForTheRuntime::StartTransaction();
 
-			AutoRTFM::RecordOpenWrite(&Data);
-			Data = Data2;
+			AutoRTFM::Assign(Data, Data2);
 			REQUIRE(Data.A == 9);
 			REQUIRE(Data.B == 8.0);
 			REQUIRE(Data.C == 7.0f);
@@ -735,8 +726,7 @@ TEST_CASE("OpenAPI.WriteTrivialStructure2")
 		{
 			AutoRTFM::ForTheRuntime::StartTransaction();
 				
-			AutoRTFM::RecordOpenWrite(&Data);
-			Data = Data2;
+			AutoRTFM::Assign(Data, Data2);
 			REQUIRE(Data.A == 9);
 			REQUIRE(Data.B == 8.0);
 			REQUIRE(Data.C == 7.0f);
@@ -747,8 +737,7 @@ TEST_CASE("OpenAPI.WriteTrivialStructure2")
 			REQUIRE(Data.E[3] == 666);
 			REQUIRE(Data.E[4] == 555);
 
-			AutoRTFM::RecordOpenWrite(&Data);
-			Data = Data3;
+			AutoRTFM::Assign(Data, Data3);
 			REQUIRE(Data.A == 19);
 			REQUIRE(Data.B == 28.0);
 			REQUIRE(Data.C == 37.0f);
diff --git a/Engine/Source/Programs/AutoRTFMTests/Private/UECoreTests.cpp b/Engine/Source/Programs/AutoRTFMTests/Private/UECoreTests.cpp
--- a/Engine/Source/Programs/AutoRTFMTests/Private/UECoreTests.cpp
+++ b/Engine/Source/Programs/AutoRTFMTests/Private/UECoreTests.cpp
@@ -876,6 +876,9 @@ TEST_CASE("UECore.FUObjectItem")
 
 	SECTION("CreateStatID In Closed Then Again In Open")
 	{
+		// This test case performs a mixed open / closed write in Item.CreateStatID()
+		AUTORTFM_SCOPED_DISABLE_MEMORY_VALIDATION();
+
 		{
 			FUObjectItem Item;
 			Item.SetObject(NewObject<UMyAutoRTFMTestObject>());
