[AutoRTFMSan] Disable the AutoRTFM sanitizer while rolling back the writes

The custom write function will want to touch the memory that is going to be in the ClosedWrite interval tree, triggering false-positives.

Also: Rename AUTORTFM_DISABLE_SANITIZER_SCOPE to AUTORTFM_SANITIZER_DISABLE_SCOPE.
This is done to have the sanitizer macros share the same prefix of AUTORTFM_SANITIZER

#Synced-CL 49102748

diff --git a/Engine/Source/Programs/AutoRTFMTests/Private/Open.cpp b/Engine/Source/Programs/AutoRTFMTests/Private/Open.cpp
--- a/Engine/Source/Programs/AutoRTFMTests/Private/Open.cpp
+++ b/Engine/Source/Programs/AutoRTFMTests/Private/Open.cpp
@@ -631,22 +631,22 @@ TEST_CASE("Open.Collision")
         });
     }
 
-    SECTION("Transact(Assign, AUTORTFM_DISABLE_SANITIZER_SCOPE(Open(Assign)))")
+    SECTION("Transact(Assign, AUTORTFM_SANITIZER_DISABLE_SCOPE(Open(Assign)))")
     {
         AutoRTFM::Transact([&]
         {
             I = 42;
 
-            AUTORTFM_DISABLE_SANITIZER_SCOPE();
+            AUTORTFM_SANITIZER_DISABLE_SCOPE();
 
             AutoRTFM::Open([&]
             {
-                I = 24; // Sanitizer error silenced by AUTORTFM_DISABLE_SANITIZER_SCOPE()
+                I = 24; // Sanitizer error silenced by AUTORTFM_SANITIZER_DISABLE_SCOPE()
             });
         });
     }
 
-    SECTION("Transact(Assign, Open(AUTORTFM_DISABLE_SANITIZER_SCOPE(Assign)))")
+    SECTION("Transact(Assign, Open(AUTORTFM_SANITIZER_DISABLE_SCOPE(Assign)))")
     {
         AutoRTFM::Transact([&]
         {
@@ -654,9 +654,9 @@ TEST_CASE("Open.Collision")
 
             AutoRTFM::Open([&]
             {
-                AUTORTFM_DISABLE_SANITIZER_SCOPE();
+                AUTORTFM_SANITIZER_DISABLE_SCOPE();
 
-                I = 24; // Sanitizer error silenced by AUTORTFM_DISABLE_SANITIZER_SCOPE()
+                I = 24; // Sanitizer error silenced by AUTORTFM_SANITIZER_DISABLE_SCOPE()
             });
         });
     }
diff --git a/Engine/Source/Runtime/AutoRTFM/Private/Transaction.cpp b/Engine/Source/Runtime/AutoRTFM/Private/Transaction.cpp
--- a/Engine/Source/Runtime/AutoRTFM/Private/Transaction.cpp
+++ b/Engine/Source/Runtime/AutoRTFM/Private/Transaction.cpp
@@ -197,8 +197,12 @@ void FTransaction::Undo()
 {
 	AUTORTFM_VERBOSE("Undoing a transaction...");
 
+	// Prevent custom rollback functions from triggering the AutoRTFM sanitizer
+	AUTORTFM_SANITIZER_DISABLE_SCOPE();
+	AUTORTFM_SANITIZER_ONLY(ClosedWrites.Reset());
+	
 	for(auto Iter = WriteLog.rbegin(); Iter != WriteLog.rend(); ++Iter)
-    {
+	{
 		FWriteLogEntry Entry = *Iter;
 		// No write records should be within the transaction's stack range.
 		AUTORTFM_ENSURE(!IsOnStack(Entry.LogicalAddress));
@@ -213,15 +217,17 @@ void FTransaction::Undo()
 		}
 	}
 
-	AUTORTFM_SANITIZER_ONLY(ClosedWrites.Reset());
-
 	AUTORTFM_VERBOSE("Undone a transaction!");
 }
 
 void FTransaction::CommitNested()
 {
-    AUTORTFM_ASSERT(Parent);
-    AUTORTFM_ASSERT(CompletionTasks == std::nullopt);
+	AUTORTFM_ASSERT(Parent);
+	AUTORTFM_ASSERT(CompletionTasks == std::nullopt);
+
+	// Disable the sanitizer for the scope of this function as we may allocate
+	// memory, which may call autortfm_sanitizer_open_write().
+	AUTORTFM_SANITIZER_DISABLE_SCOPE();
 
 	// We need to pass our write log to our parent transaction, but with care!
 	// We need to discard any writes if the memory location is on the parent
diff --git a/Engine/Source/Runtime/AutoRTFM/Public/AutoRTFM.h b/Engine/Source/Runtime/AutoRTFM/Public/AutoRTFM.h
--- a/Engine/Source/Runtime/AutoRTFM/Public/AutoRTFM.h
+++ b/Engine/Source/Runtime/AutoRTFM/Public/AutoRTFM.h
@@ -214,10 +214,10 @@ namespace Sanitizer
 #endif // AUTORTFM_SANITIZER
 
 #if AUTORTFM_SANITIZER
-	#define AUTORTFM_DISABLE_SANITIZER_SCOPE() \
+	#define AUTORTFM_SANITIZER_DISABLE_SCOPE() \
 		::AutoRTFM::Sanitizer::FDisableScope UE_AUTORTFM_CONCAT(AutoRTFMSanitizerModeScope, __COUNTER__)
 #else
-	#define AUTORTFM_DISABLE_SANITIZER_SCOPE() static_assert(true) /* require semicolon */
+	#define AUTORTFM_SANITIZER_DISABLE_SCOPE() static_assert(true) /* require semicolon */
 #endif
 
 #if UE_AUTORTFM
@@ -794,7 +794,7 @@ AUTORTFM_DISABLE_NO_SANITIZE void Assign(RefType& Ref, ValueType&& NewValue, EWr
 {
 	autortfm_record_open_write_with_flags(&Ref, sizeof(Ref), static_cast<autortfm_write_flags>(Flags));
 	{
-		AUTORTFM_DISABLE_SANITIZER_SCOPE();
+		AUTORTFM_SANITIZER_DISABLE_SCOPE();
 		Ref = std::forward<ValueType>(NewValue);
 	}
 }
