[AutoRTFM] Restructure OnPreAbort tests to run with the sanitizer enabled.

Following up on some code review discussion from earlier. Unfortunately, I am not sure how to enable the sanitizer on my machine (`__has_feature(autortfm_sanitizer)` is apparently false) but conceptually this should prevent us from making mixed open/closed writes to the same location.

#Synced-CL 49134418

diff --git a/Engine/Source/Programs/AutoRTFMTests/Private/API.cpp b/Engine/Source/Programs/AutoRTFMTests/Private/API.cpp
--- a/Engine/Source/Programs/AutoRTFMTests/Private/API.cpp
+++ b/Engine/Source/Programs/AutoRTFMTests/Private/API.cpp
@@ -1276,21 +1276,26 @@ TEST_CASE("API.OnAbortMacro_NoAbort")
 
 TEST_CASE("API.OnPreAbortMacro_NoAbort")
 {
-	AUTORTFM_SCOPED_SANITIZER_MODE_DISABLE();
+	AUTORTFM_SCOPED_DISABLE_RETRY();
 
 	int Value = 123;
+	bool bReachedOnPreAbort = false;
 
 	UE_AUTORTFM_TRANSACT
 	{
 		Value = 456;
 
-		UE_AUTORTFM_ONPREABORT(&Value)
+		UE_AUTORTFM_ONPREABORT(&Value, &bReachedOnPreAbort)
 		{
-			Value = 123;
+			UE_AUTORTFM_OPEN
+			{
+				bReachedOnPreAbort = true;
+			};
 		};
 	};
 
 	REQUIRE(Value == 456);
+	REQUIRE(!bReachedOnPreAbort);
 }
 
 TEST_CASE("API.OnAbortMacro_WithAbort")
@@ -1301,7 +1306,7 @@ TEST_CASE("API.OnAbortMacro_WithAbort")
 	{
 		Value = 234;
 
-		UE_AUTORTFM_ONABORT(&Value)
+		UE_AUTORTFM_ONABORT(&Value) 
 		{
 			Value = 123;
 		};
@@ -1314,23 +1319,34 @@ TEST_CASE("API.OnAbortMacro_WithAbort")
 
 TEST_CASE("API.OnPreAbortMacro_WithAbort")
 {
-	AUTORTFM_SCOPED_SANITIZER_MODE_DISABLE();
-
 	int Value = 123;
+	bool bReachedOnPreAbort = false;
 
 	UE_AUTORTFM_TRANSACT
 	{
 		Value = 234;
 
-		UE_AUTORTFM_ONPREABORT(&Value)
+		UE_AUTORTFM_ONPREABORT(&Value, &bReachedOnPreAbort)
 		{
-			Value = 123;
+			// In OnPreAbort, the value is still changed.
+			REQUIRE(Value == 234);
+			UE_AUTORTFM_OPEN
+			{
+				bReachedOnPreAbort = true;
+			};
+		};
+
+		UE_AUTORTFM_ONABORT(&Value)
+		{
+			// In OnAbort, the value has been rolled back.
+			REQUIRE(Value == 123);
 		};
 
 		AutoRTFM::AbortTransaction();
 	};
 
 	REQUIRE(Value == 123);
+	REQUIRE(bReachedOnPreAbort);
 }
 
 TEST_CASE("API.DidAllocate")
