Make `TLinearAllocatorBase` AutoRTFM-safe.

#Synced-CL 49675120

diff --git a/Engine/Source/Programs/AutoRTFMTests/Private/UECoreTests.cpp b/Engine/Source/Programs/AutoRTFMTests/Private/UECoreTests.cpp
--- a/Engine/Source/Programs/AutoRTFMTests/Private/UECoreTests.cpp
+++ b/Engine/Source/Programs/AutoRTFMTests/Private/UECoreTests.cpp
@@ -9,6 +9,7 @@
 #include "AutoRTFM/Testing.h"
 #include "Blueprint/BlueprintExceptionInfo.h"
 #include "Delegates/IDelegateInstance.h"
+#include "Experimental/ConcurrentLinearAllocator.h"
 #include "Experimental/UnifiedError/CoreErrorTypes.h"
 #include "HAL/IConsoleManager.h"
 #include "HAL/MallocLeakDetection.h"
@@ -2381,3 +2382,67 @@ TEST_CASE("UECore.ArchiveObjectCrc32")
 		});
 	}
 }
+
+TEST_CASE("UECore.FConcurrentLinearAllocator")
+{
+	SECTION("Malloc")
+	{
+		int32* Allocation = nullptr;
+
+		AutoRTFM::Testing::Commit([&]
+		{
+			Allocation = static_cast<int32*>(FConcurrentLinearAllocator::Malloc(sizeof(int32), alignof(int32)));
+			*Allocation = 42;
+		});
+
+		REQUIRE(Allocation);
+		REQUIRE(*Allocation == 42);
+
+		FConcurrentLinearAllocator::Free(Allocation);
+	}
+
+	SECTION("Free")
+	{
+		int32* Allocation = static_cast<int32*>(FConcurrentLinearAllocator::Malloc(sizeof(int32), alignof(int32)));
+		REQUIRE(Allocation);
+
+		AutoRTFM::Testing::Commit([&]
+		{
+			FConcurrentLinearAllocator::Free(Allocation);
+		});
+	}
+
+	SECTION("Realloc")
+	{
+		int32* Allocation = static_cast<int32*>(FConcurrentLinearAllocator::Malloc(sizeof(int32), alignof(int32)));
+		REQUIRE(Allocation);
+
+		Allocation[0] = 42;
+
+		AutoRTFM::Testing::Commit([&]
+		{
+			Allocation = static_cast<int32*>(FConcurrentLinearAllocator::Realloc(Allocation, 2 * sizeof(int32), alignof(int32)));
+			Allocation[1] = 13;
+		});
+
+		REQUIRE(Allocation);
+		REQUIRE(Allocation[0] == 42);
+		REQUIRE(Allocation[1] == 13);
+		FConcurrentLinearAllocator::Free(Allocation);
+	}
+
+	SECTION("GetAllocationSize")
+	{
+		int32* Allocation = static_cast<int32*>(FConcurrentLinearAllocator::Malloc(sizeof(int32), alignof(int32)));
+		REQUIRE(Allocation);
+
+		const size_t AllocationSize = FConcurrentLinearAllocator::GetAllocationSize(Allocation);
+
+		AutoRTFM::Testing::Commit([&]
+		{
+			REQUIRE(AllocationSize == FConcurrentLinearAllocator::GetAllocationSize(Allocation));
+		});
+
+		FConcurrentLinearAllocator::Free(Allocation);
+	}
+}
diff --git a/Engine/Source/Runtime/Core/Public/Experimental/ConcurrentLinearAllocator.h b/Engine/Source/Runtime/Core/Public/Experimental/ConcurrentLinearAllocator.h
--- a/Engine/Source/Runtime/Core/Public/Experimental/ConcurrentLinearAllocator.h
+++ b/Engine/Source/Runtime/Core/Public/Experimental/ConcurrentLinearAllocator.h
@@ -5,6 +5,7 @@
 #include <atomic>
 #include <type_traits>
 
+#include "AutoRTFM.h"
 #include "HAL/Memory.h"
 #include "HAL/UnrealMemory.h"
 #include "HAL/LowLevelMemTracker.h"
@@ -347,6 +348,22 @@ public:
 	static inline void* Malloc(SIZE_T Size, uint32 Alignment)
 	{
 		checkSlow(Alignment >= 1 && FMath::IsPowerOfTwo(Alignment));
+
+		if (AutoRTFM::IsClosed())
+		{
+			void* const Pointer = AutoRTFM::Open([Size, Alignment]
+			{
+				return Malloc(Size, Alignment);
+			});
+
+			AutoRTFM::OnAbort([Pointer]
+			{
+				Free(Pointer);
+			});
+
+			return Pointer;
+		}
+
 		if constexpr (!SupportsFastPath)
 		{
 			Alignment = (Alignment < alignof(FAllocationHeader)) ? alignof(FAllocationHeader) : Alignment; 
@@ -486,6 +503,16 @@ public:
 	{
 		if(Pointer != nullptr)
 		{
+			if (AutoRTFM::IsClosed())
+			{
+				AutoRTFM::OnCommit([Pointer]
+				{
+					Free(Pointer);
+				});
+
+				return;
+			}
+
 			if constexpr (SupportsFastPath)
 			{
 				MemoryTrace_Free(uint64(Pointer));
@@ -523,7 +550,7 @@ public:
 		}
 	}
 
-	static inline SIZE_T GetAllocationSize(void* Pointer)
+	AUTORTFM_OPEN static inline SIZE_T GetAllocationSize(void* Pointer)
 	{
 		if (Pointer)
 		{
