Get `FUniversallyUniqueObjectNameGenerator` to work under transactions.

#Synced-CL 49675120

diff --git a/Engine/Source/Programs/AutoRTFMTests/Private/UObjectTests.cpp b/Engine/Source/Programs/AutoRTFMTests/Private/UObjectTests.cpp
--- a/Engine/Source/Programs/AutoRTFMTests/Private/UObjectTests.cpp
+++ b/Engine/Source/Programs/AutoRTFMTests/Private/UObjectTests.cpp
@@ -1387,3 +1387,54 @@ TEST_CASE("UObject.ObjectWithSubObject")
 		REQUIRE(Object);
 	});
 }
+
+TEST_CASE("UObject.FUniversallyUniqueObjectNameGenerator")
+{
+	SECTION("Constructed outside of transaction")
+	{
+		FGuid Seed;
+		FUniversallyUniqueObjectNameGenerator Generator(Seed);
+
+		const uint32 Value = Generator.GetNewValue();
+
+		AutoRTFM::Testing::Commit([&]
+		{
+			REQUIRE(Generator.GetNewValue() != Value);
+		});
+	}
+
+	SECTION("Constructed inside of transaction")
+	{
+		AutoRTFM::Testing::Commit([&]
+		{
+			FGuid Seed;
+			FUniversallyUniqueObjectNameGenerator Generator(Seed);
+			uint32 Value = Generator.GetNewValue();
+
+			// Just check that we are definitely generating the pseudo-random values even in transactions.
+			REQUIRE(Value != Generator.GetNewValue());
+		});
+	}
+
+	SECTION("Seed is not reset by aborting transaction")
+	{
+		FGuid Seed;
+		FUniversallyUniqueObjectNameGenerator Generator(Seed);
+
+		uint32 Value = 0;
+
+		AutoRTFM::Testing::Abort([&]
+		{
+			const uint32 TransactionalValue = Generator.GetNewValue();
+
+			AutoRTFM::OnAbort([&Value, TransactionalValue]
+			{
+				Value = TransactionalValue;
+			});
+
+			AutoRTFM::AbortTransaction();
+		});
+
+		REQUIRE(Value != Generator.GetNewValue());
+	}
+}
diff --git a/Engine/Source/Runtime/CoreUObject/Public/UObject/UObjectGlobals.h b/Engine/Source/Runtime/CoreUObject/Public/UObject/UObjectGlobals.h
--- a/Engine/Source/Runtime/CoreUObject/Public/UObject/UObjectGlobals.h
+++ b/Engine/Source/Runtime/CoreUObject/Public/UObject/UObjectGlobals.h
@@ -1046,7 +1046,7 @@ COREUOBJECT_API FName MakeUniqueObjectName( UObject* Outer, const UClass* Class,
 /**
  * Generator class that can be used as a session to generate universally unique object names.
  */
-struct FUniversallyUniqueObjectNameGenerator
+struct AUTORTFM_OPEN FUniversallyUniqueObjectNameGenerator
 {
 	COREUOBJECT_API FUniversallyUniqueObjectNameGenerator(const FGuid& Seed);
 	COREUOBJECT_API uint32 GetNewValue() const;
