[AutoRTFMTests] Disable FMallocLeakDetector

This is fighting with the new AutoRTFM::Testing::FTrackingAllocator, and is causing runs to non-deterministically fail.
The FTrackingAllocator already identifes memory leaks in AutoRTFM, so just disable FMallocLeakDetector

#Synced-CL 49993096

diff --git a/Engine/Source/Programs/AutoRTFMTests/AutoRTFMTests.Target.cs b/Engine/Source/Programs/AutoRTFMTests/AutoRTFMTests.Target.cs
--- a/Engine/Source/Programs/AutoRTFMTests/AutoRTFMTests.Target.cs
+++ b/Engine/Source/Programs/AutoRTFMTests/AutoRTFMTests.Target.cs
@@ -71,7 +71,6 @@ public class AutoRTFMTestsTarget : TargetRules
 			GlobalDefinitions.Add("ALLOW_LOG_FILE=1");
 		}
 
-		GlobalDefinitions.Add("MALLOC_LEAKDETECTION=1");
 		GlobalDefinitions.Add("PLATFORM_USES_FIXED_GMalloc_CLASS=0");
 
 		GlobalDefinitions.Add("UNIX_PLATFORM_FILE_SPEEDUP_FILE_OPERATIONS=1");
diff --git a/Engine/Source/Programs/AutoRTFMTests/Private/AutoRTFMTests.cpp b/Engine/Source/Programs/AutoRTFMTests/Private/AutoRTFMTests.cpp
--- a/Engine/Source/Programs/AutoRTFMTests/Private/AutoRTFMTests.cpp
+++ b/Engine/Source/Programs/AutoRTFMTests/Private/AutoRTFMTests.cpp
@@ -8,7 +8,6 @@
 #include "Misc/LazySingleton.h"
 #include "Misc/PackageName.h"
 #include "RequiredProgramMainCPPInclude.h" // required for ue programs
-#include "HAL/MallocLeakDetection.h"
 #include "UObject/GCObject.h"
 #include "Serialization/PackageStore.h"
 
@@ -17,49 +16,15 @@ IMPLEMENT_APPLICATION(AutoRTFMTests, "AutoRTFMTests");
 #define CATCH_AMALGAMATED_CUSTOM_MAIN
 #include "catch_amalgamated.cpp"
 
-class FListener final : public Catch::EventListenerBase
-{
-public:
-	using Catch::EventListenerBase::EventListenerBase;
-
-	void testCaseStarting(const Catch::TestCaseInfo&) override
-	{
-		FMallocLeakDetection::Get().SetAllocationCollection(true);
-	}
-
-	void testCaseEnded(const Catch::TestCaseStats&) override
-	{
-		FMallocLeakDetection::Get().SetAllocationCollection(false);
-	}
-
-	void testRunEnded(const Catch::TestRunStats&) override
-	{
-	}
-};
-
-CATCH_REGISTER_LISTENER(FListener)
-
 void AutoRTFM::Testing::AssertionFailure(const char* Expression, const char* File, int Line)
 {
 	FAIL(File << ":" << Line << ": " << Expression);
 }
 
-// Checks for memory leaks. Returns true if no leaks were found, otherwise all
-// leaks are printed to stderr and false is returned.
-bool CheckNoMemoryLeaks();
-
 int RunTests(int ArgC, const char* ArgV[])
 {
 	FTaskTagScope Scope(ETaskTag::EGameThread);
 
-	{
-		// measure_environment lazily allocates an Environment pointer on first
-		// call, and holds this in a static variable. Pre-allocate this with the
-		// leak detector disabled to prevent this being reported as a leak.
-		MALLOCLEAK_IGNORE_SCOPE();
-		Catch::Benchmark::Detail::measure_environment<Catch::Benchmark::default_clock>();
-	}
-
 	Catch::Session Session;
 
 	bool bNoRetry = false;
