Fix UPackage tests

#Synced-CL 49950557

diff --git a/Engine/Source/Programs/AutoRTFMTests/Private/UPackageTests.cpp b/Engine/Source/Programs/AutoRTFMTests/Private/UPackageTests.cpp
--- a/Engine/Source/Programs/AutoRTFMTests/Private/UPackageTests.cpp
+++ b/Engine/Source/Programs/AutoRTFMTests/Private/UPackageTests.cpp
@@ -571,7 +571,7 @@ TEST_CASE("UPackage.AsyncLoading")
 
 		AutoRTFM::Testing::Commit([&]
 		{
-			FName CustomName("Wowwee");
+			FName CustomName("CustomName", 1);
 
 			FLoadPackageAsyncDelegate Delegate;
 
@@ -600,7 +600,7 @@ TEST_CASE("UPackage.AsyncLoading")
 		{
 			TrashObject(Package);
 
-			FName CustomName("Wowwee");
+			FName CustomName("CustomName", 2);
 
 			FLoadPackageAsyncDelegate Delegate;
 
@@ -628,7 +628,7 @@ TEST_CASE("UPackage.AsyncLoading")
 
 		bool bFirst = true;
 
-		FName CustomName("Wowwee");
+		FName CustomName("CustomName", 3);
 
 		FLoadPackageAsyncDelegate Delegate;
 
