[AutoRTFM] Add Log2 utility function.

This CL also removes the unused `Lerp` utility function.

This is based on the unfinished CL at https://p4-swarm.epicgames.net/reviews/49991893

#Synced-CL 49876485

diff --git a/Engine/Source/Programs/AutoRTFMTests/Private/UtilsTests.cpp b/Engine/Source/Programs/AutoRTFMTests/Private/UtilsTests.cpp
--- a/Engine/Source/Programs/AutoRTFMTests/Private/UtilsTests.cpp
+++ b/Engine/Source/Programs/AutoRTFMTests/Private/UtilsTests.cpp
@@ -255,11 +255,42 @@ TEST_CASE("RoundUp")
 	REQUIRE(AutoRTFM::RoundUp(8, 8) == 8);
 }
 
-TEST_CASE("Lerp")
+TEST_CASE("Log2")
 {
-	REQUIRE_THAT(AutoRTFM::Lerp(2.0, 4.0, -0.5), Catch::Matchers::WithinRel(1.0, 1e-5));
-	REQUIRE_THAT(AutoRTFM::Lerp(2.0, 4.0, +0.0), Catch::Matchers::WithinRel(2.0, 1e-5));
-	REQUIRE_THAT(AutoRTFM::Lerp(2.0, 4.0, +0.5), Catch::Matchers::WithinRel(3.0, 1e-5));
-	REQUIRE_THAT(AutoRTFM::Lerp(2.0, 4.0, +1.0), Catch::Matchers::WithinRel(4.0, 1e-5));
-	REQUIRE_THAT(AutoRTFM::Lerp(2.0, 4.0, +1.5), Catch::Matchers::WithinRel(5.0, 1e-5));
-}
+	static_assert(AutoRTFM::Log2(0u) == 0);
+	static_assert(AutoRTFM::Log2(1u) == 0);
+	static_assert(AutoRTFM::Log2(2u) == 1);
+	static_assert(AutoRTFM::Log2(3u) == 1);
+	static_assert(AutoRTFM::Log2(4u) == 2);
+	static_assert(AutoRTFM::Log2(5u) == 2);
+	static_assert(AutoRTFM::Log2(0x0fffu) == 11);
+	static_assert(AutoRTFM::Log2(0x1000u) == 12);
+	static_assert(AutoRTFM::Log2(0x1001u) == 12);
+	static_assert(AutoRTFM::Log2<uint8_t>(0x0) == 0);
+	static_assert(AutoRTFM::Log2<uint8_t>(0xf0) == 7);
+	static_assert(AutoRTFM::Log2<uint16_t>(0x0) == 0);
+	static_assert(AutoRTFM::Log2<uint16_t>(0xcdef) == 15);
+	static_assert(AutoRTFM::Log2<uint32_t>(0x7fffffff) == 30);
+	static_assert(AutoRTFM::Log2<uint32_t>(0x80000000) == 31);
+	static_assert(AutoRTFM::Log2<uint32_t>(0x80000001) == 31);
+	static_assert(AutoRTFM::Log2<uint32_t>(0xfffffffe) == 31);
+	static_assert(AutoRTFM::Log2<uint32_t>(0xffffffff) == 31);
+	static_assert(AutoRTFM::Log2<uint64_t>(0x7fffffffffffffff) == 62);
+	static_assert(AutoRTFM::Log2<uint64_t>(0x8000000000000000) == 63);
+	static_assert(AutoRTFM::Log2<uint64_t>(0x8000000000000001) == 63);
+	static_assert(AutoRTFM::Log2<uint64_t>(0xfffffffffffffffe) == 63);
+	static_assert(AutoRTFM::Log2<uint64_t>(0xffffffffffffffff) == 63);
+
+	for (int Index = 1; Index <= 8; ++Index)
+	{
+		int LowValue  = (1 << Index) / 2;
+		REQUIRE(AutoRTFM::Log2<uint8_t> (LowValue) == Index-1);
+		REQUIRE(AutoRTFM::Log2<uint16_t>(LowValue) == Index-1);
+		REQUIRE(AutoRTFM::Log2<uint32_t>(LowValue) == Index-1);
+
+		int HighValue = (1 << Index) - 1;
+		REQUIRE(AutoRTFM::Log2<uint8_t> (HighValue) == Index-1);
+		REQUIRE(AutoRTFM::Log2<uint16_t>(HighValue) == Index-1);
+		REQUIRE(AutoRTFM::Log2<uint32_t>(HighValue) == Index-1);
+	}
+}
\ No newline at end of file
diff --git a/Engine/Source/Runtime/AutoRTFM/Private/Utils.h b/Engine/Source/Runtime/AutoRTFM/Private/Utils.h
--- a/Engine/Source/Runtime/AutoRTFM/Private/Utils.h
+++ b/Engine/Source/Runtime/AutoRTFM/Private/Utils.h
@@ -164,11 +164,13 @@ inline constexpr T RoundUp(T Value, T Multiple)
 	return RoundDown(Value + Multiple - 1, Multiple);
 }
 
-// Returns the linear interpolation between Start and End using the coefficient Fraction.
+// Returns the log2 of Value
 template <typename T>
-inline constexpr T Lerp(T Start, T End, T Fraction)
+requires(std::is_integral_v<T> && std::is_unsigned_v<T>)
+inline constexpr T Log2(T Value)
 {
-	return Start + Fraction * (End - Start);
+	constexpr size_t NumBits = sizeof(Value) * 8;
+	return Value > 0 ? NumBits - __builtin_clzg(Value) - 1 : 0;
 }
 
 }  // namespace AutoRTFM
