[AutoRTFM] Track the number of bytes allocated via the AutoRTFM ExternAPI.

This is based on the unfinished CL at https://p4-swarm.epicgames.net/reviews/49991893 with some modifications based on my CL review comments. Additionally, there is now a simple unit test not present in the original work.

I've added clarification lines to the ExternAPI doc comments in order to steer Bughawk. It was concerned that the memory stats might be incorrect in cases where the allocation APIs failed and returned null. This is a reasonable concern if you don't know that the ExternAPI enforces allocations to succeed.

#Synced-CL 49876485

diff --git a/Engine/Source/Programs/AutoRTFMTests/Private/Allocators.cpp b/Engine/Source/Programs/AutoRTFMTests/Private/Allocators.cpp
--- a/Engine/Source/Programs/AutoRTFMTests/Private/Allocators.cpp
+++ b/Engine/Source/Programs/AutoRTFMTests/Private/Allocators.cpp
@@ -3,6 +3,8 @@
 #include "AutoRTFM.h"
 #include "AutoRTFM/Testing.h"
 #include "Catch2Includes.h"
+#include "ExternAPI.h"
+#include "MemoryStats.h"
 
 #include <new>
 #include <stdlib.h>
@@ -369,6 +371,31 @@ TEST_CASE("Allocators.Realloc")
 	MSVC_END_DISABLE_WARN_ALLOC_MAYBE_NULL
 }
 
+TEST_CASE("Allocators.MemoryStats.GetTotalBytesAllocated")
+{
+	size_t StartingSize = AutoRTFM::FMemoryStats::GetTotalBytesAllocated();
+
+	void* const Allocate1MBytes = AutoRTFM::Allocate(1'000'000, /*Alignment*/ 1);
+	size_t CurrentSize = AutoRTFM::FMemoryStats::GetTotalBytesAllocated();
+	REQUIRE(StartingSize + 1'000'000 == CurrentSize);
+
+	void* const Allocate49MBytes = AutoRTFM::AllocateZeroed(49'000'000, /*Alignment*/ 1);
+	CurrentSize = AutoRTFM::FMemoryStats::GetTotalBytesAllocated();
+	REQUIRE(StartingSize + 50'000'000 == CurrentSize);
+
+	AutoRTFM::Free(Allocate1MBytes, 1'000'000);
+	CurrentSize = AutoRTFM::FMemoryStats::GetTotalBytesAllocated();
+	REQUIRE(StartingSize + 49'000'000 == CurrentSize);
+
+	void* const Allocate7MBytes = AutoRTFM::Reallocate(Allocate49MBytes, 7'000'000, /*Alignment*/ 1, /*PreviousSize*/ 49'000'000);
+	CurrentSize = AutoRTFM::FMemoryStats::GetTotalBytesAllocated();
+	REQUIRE(StartingSize + 7'000'000 == CurrentSize);
+
+	AutoRTFM::Free(Allocate7MBytes, 7'000'000);
+	CurrentSize = AutoRTFM::FMemoryStats::GetTotalBytesAllocated();
+	REQUIRE(CurrentSize == StartingSize);
+}
+
 #undef CLANG_BEGIN_DISABLE_OPTIMIZATIONS
 #undef CLANG_END_DISABLE_OPTIMIZATIONS
 #undef MSVC_BEGIN_DISABLE_WARN_ALLOC_MAYBE_NULL
diff --git a/Engine/Source/Runtime/AutoRTFM/Private/ExternAPI.h b/Engine/Source/Runtime/AutoRTFM/Private/ExternAPI.h
--- a/Engine/Source/Runtime/AutoRTFM/Private/ExternAPI.h
+++ b/Engine/Source/Runtime/AutoRTFM/Private/ExternAPI.h
@@ -4,7 +4,8 @@
 
 #if (defined(__AUTORTFM) && __AUTORTFM)
 
-#include "AutoRTFM.h"
+#include "AutoRTFM/CAPI.h"
+#include "MemoryStats.h"
 
 namespace AutoRTFM
 {
@@ -13,21 +14,25 @@ UE_AUTORTFM_API extern autortfm_extern_api GExternAPI;
 
 inline void* Allocate(size_t Size, size_t Alignment)
 {
+	FMemoryStats::UpdateBytesAllocated(Size);
 	return GExternAPI.Allocate(Size, Alignment);
 }
 
 inline void* Reallocate(void* Pointer, size_t Size, size_t Alignment, size_t PreviousSize)
 {
+	FMemoryStats::UpdateBytesAllocated(static_cast<ptrdiff_t>(Size) - static_cast<ptrdiff_t>(PreviousSize));
 	return GExternAPI.Reallocate(Pointer, Size, Alignment, PreviousSize);
 }
 
 inline void* AllocateZeroed(size_t Size, size_t Alignment)
 {
+	FMemoryStats::UpdateBytesAllocated(Size);
 	return GExternAPI.AllocateZeroed(Size, Alignment);
 }
 
 inline void Free(void* Pointer, size_t AllocationSize)
 {
+	FMemoryStats::UpdateBytesAllocated(-static_cast<ptrdiff_t>(AllocationSize));
 	return GExternAPI.Free(Pointer, AllocationSize);
 }
 
diff --git a/Engine/Source/Runtime/AutoRTFM/Private/MemoryStats.cpp b/Engine/Source/Runtime/AutoRTFM/Private/MemoryStats.cpp
new file mode 100644
--- /dev/null
+++ b/Engine/Source/Runtime/AutoRTFM/Private/MemoryStats.cpp
@@ -0,0 +1,22 @@
+// Copyright Epic Games, Inc. All Rights Reserved.
+
+#if (defined(__AUTORTFM) && __AUTORTFM)
+
+#include "MemoryStats.h"
+#include "Utils.h"
+
+namespace AutoRTFM
+{
+
+std::atomic<size_t> FMemoryStats::BytesAllocated = 0;
+
+#if AUTORTFM_CHECK_FOR_ALLOCATION_UNDERFLOW
+void FMemoryStats::FatalErrorBytesAllocatedUnderflow()
+{
+	AUTORTFM_FATAL("FMemoryStats::BytesAllocated has underflowed");
+}
+#endif
+
+}  // namespace AutoRTFM
+
+#endif  // defined(__AUTORTFM) && __AUTORTFM
diff --git a/Engine/Source/Runtime/AutoRTFM/Private/MemoryStats.h b/Engine/Source/Runtime/AutoRTFM/Private/MemoryStats.h
new file mode 100644
--- /dev/null
+++ b/Engine/Source/Runtime/AutoRTFM/Private/MemoryStats.h
@@ -0,0 +1,56 @@
+// Copyright Epic Games, Inc. All Rights Reserved.
+
+#pragma once
+
+#if (defined(__AUTORTFM) && __AUTORTFM)
+
+#include "AutoRTFM.h"
+#include "BuildMacros.h"
+#include <atomic>
+
+namespace AutoRTFM
+{
+
+#if AUTORTFM_BUILD_SHIPPING
+#define AUTORTFM_CHECK_FOR_ALLOCATION_UNDERFLOW 1
+#else
+#define AUTORTFM_CHECK_FOR_ALLOCATION_UNDERFLOW 0
+#endif
+
+struct FMemoryStats
+{
+	// size_t constant with only the most-significant-bit set.
+	// Used to quickly check for underflow/overflow, as we don't expect over
+	// nine quintillion bytes of heap memory to be allocated.
+	static constexpr size_t UnderflowMask = ((~static_cast<size_t>(0)) >> 1) + 1;
+
+	/** Adjusts the total number of bytes allocated by Delta. */
+	static inline void UpdateBytesAllocated(ptrdiff_t Delta)
+	{
+		[[maybe_unused]] size_t const NewBytesAllocated = (BytesAllocated += Delta);
+
+#if AUTORTFM_CHECK_FOR_ALLOCATION_UNDERFLOW
+		if (AUTORTFM_UNLIKELY(NewBytesAllocated & FMemoryStats::UnderflowMask))
+		{
+			FatalErrorBytesAllocatedUnderflow();
+		}
+#endif
+	}
+
+	/** Returns the total number of heap bytes allocated by AutoRTFM. */
+	static inline size_t GetTotalBytesAllocated()
+	{
+		return BytesAllocated;
+	}
+
+private:
+#if AUTORTFM_CHECK_FOR_ALLOCATION_UNDERFLOW
+	void FatalErrorBytesAllocatedUnderflow();  // Not inline to break circular header dependencies.
+#endif
+
+	UE_AUTORTFM_API static std::atomic<size_t> BytesAllocated;
+};
+
+}  // namespace AutoRTFM
+
+#endif  // (defined(__AUTORTFM) && __AUTORTFM)
diff --git a/Engine/Source/Runtime/AutoRTFM/Public/AutoRTFM/CAPI.h b/Engine/Source/Runtime/AutoRTFM/Public/AutoRTFM/CAPI.h
--- a/Engine/Source/Runtime/AutoRTFM/Public/AutoRTFM/CAPI.h
+++ b/Engine/Source/Runtime/AutoRTFM/Public/AutoRTFM/CAPI.h
@@ -102,19 +102,19 @@ extern "C"
 	// Function pointers used by AutoRTFM for heap allocations, etc.
 	typedef struct
 	{
-		// The function used to allocate memory from the heap.
-		// Must not be null.
+		// The function used to allocate memory from the heap. Must not be null.
+		// This API should never return null; if the allocator cannot satisfy the request, it must abort the program.
 		void* (*Allocate)(size_t Size, size_t Alignment);
 
-		// The function used to reallocate memory from the heap.
-		// Must not be null.
+		// The function used to reallocate memory from the heap. Must not be null.
+		// This API should never return null; if the allocator cannot satisfy the request, it must abort the program.
 		void* (*Reallocate)(void* Pointer, size_t Size, size_t Alignment, size_t PreviousSize);
 
-		// The function used to allocate zeroed memory from the heap.
-		// Must not be null.
+		// The function used to allocate zeroed memory from the heap. Must not be null.
+		// This API should never return null; if the allocator cannot satisfy the request, it must abort the program.
 		void* (*AllocateZeroed)(size_t Size, size_t Alignment);
 
-		// The function used to free memory allocated by Allocate(), Reallocate() and AllocateZeroed().
+		// The function used to free memory allocated by Allocate(), Reallocate() or AllocateZeroed().
 		// Must not be null.
 		void (*Free)(void* Pointer, size_t AllocationSize);
 
