[Backout] - CL50103460
#rnx
Original CL Desc
-----------------------------------------------------------------
[AutoRTFM] Track the number of bytes allocated via the AutoRTFM ExternAPI.

This is based on the unfinished CL at https://p4-swarm.epicgames.net/reviews/49991893 with some modifications based on my CL review comments. Additionally, there is now a simple unit test not present in the original work.

I've added clarification lines to the ExternAPI doc comments in order to steer Bughawk. It was concerned that the memory stats might be incorrect in cases where the allocation APIs failed and returned null. This is a reasonable concern if you don't know that the ExternAPI enforces allocations to succeed.

#Synced-CL 49876485

diff --git a/Engine/Source/Programs/AutoRTFMTests/Private/Allocators.cpp b/Engine/Source/Programs/AutoRTFMTests/Private/Allocators.cpp
--- a/Engine/Source/Programs/AutoRTFMTests/Private/Allocators.cpp
+++ b/Engine/Source/Programs/AutoRTFMTests/Private/Allocators.cpp
@@ -3,8 +3,6 @@
 #include "AutoRTFM.h"
 #include "AutoRTFM/Testing.h"
 #include "Catch2Includes.h"
-#include "ExternAPI.h"
-#include "MemoryStats.h"
 
 #include <new>
 #include <stdlib.h>
@@ -371,31 +369,6 @@ TEST_CASE("Allocators.Realloc")
 	MSVC_END_DISABLE_WARN_ALLOC_MAYBE_NULL
 }
 
-TEST_CASE("Allocators.MemoryStats.GetTotalBytesAllocated")
-{
-	size_t StartingSize = AutoRTFM::FMemoryStats::GetTotalBytesAllocated();
-
-	void* const Allocate1MBytes = AutoRTFM::Allocate(1'000'000, /*Alignment*/ 1);
-	size_t CurrentSize = AutoRTFM::FMemoryStats::GetTotalBytesAllocated();
-	REQUIRE(StartingSize + 1'000'000 == CurrentSize);
-
-	void* const Allocate49MBytes = AutoRTFM::AllocateZeroed(49'000'000, /*Alignment*/ 1);
-	CurrentSize = AutoRTFM::FMemoryStats::GetTotalBytesAllocated();
-	REQUIRE(StartingSize + 50'000'000 == CurrentSize);
-
-	AutoRTFM::Free(Allocate1MBytes, 1'000'000);
-	CurrentSize = AutoRTFM::FMemoryStats::GetTotalBytesAllocated();
-	REQUIRE(StartingSize + 49'000'000 == CurrentSize);
-
-	void* const Allocate7MBytes = AutoRTFM::Reallocate(Allocate49MBytes, 7'000'000, /*Alignment*/ 1, /*PreviousSize*/ 49'000'000);
-	CurrentSize = AutoRTFM::FMemoryStats::GetTotalBytesAllocated();
-	REQUIRE(StartingSize + 7'000'000 == CurrentSize);
-
-	AutoRTFM::Free(Allocate7MBytes, 7'000'000);
-	CurrentSize = AutoRTFM::FMemoryStats::GetTotalBytesAllocated();
-	REQUIRE(CurrentSize == StartingSize);
-}
-
 #undef CLANG_BEGIN_DISABLE_OPTIMIZATIONS
 #undef CLANG_END_DISABLE_OPTIMIZATIONS
 #undef MSVC_BEGIN_DISABLE_WARN_ALLOC_MAYBE_NULL
diff --git a/Engine/Source/Runtime/AutoRTFM/Private/ExternAPI.h b/Engine/Source/Runtime/AutoRTFM/Private/ExternAPI.h
--- a/Engine/Source/Runtime/AutoRTFM/Private/ExternAPI.h
+++ b/Engine/Source/Runtime/AutoRTFM/Private/ExternAPI.h
@@ -4,8 +4,7 @@
 
 #if (defined(__AUTORTFM) && __AUTORTFM)
 
-#include "AutoRTFM/CAPI.h"
-#include "MemoryStats.h"
+#include "AutoRTFM.h"
 
 namespace AutoRTFM
 {
@@ -14,25 +13,21 @@ UE_AUTORTFM_API extern autortfm_extern_api GExternAPI;
 
 inline void* Allocate(size_t Size, size_t Alignment)
 {
-	FMemoryStats::UpdateBytesAllocated(Size);
 	return GExternAPI.Allocate(Size, Alignment);
 }
 
 inline void* Reallocate(void* Pointer, size_t Size, size_t Alignment, size_t PreviousSize)
 {
-	FMemoryStats::UpdateBytesAllocated(static_cast<ptrdiff_t>(Size) - static_cast<ptrdiff_t>(PreviousSize));
 	return GExternAPI.Reallocate(Pointer, Size, Alignment, PreviousSize);
 }
 
 inline void* AllocateZeroed(size_t Size, size_t Alignment)
 {
-	FMemoryStats::UpdateBytesAllocated(Size);
 	return GExternAPI.AllocateZeroed(Size, Alignment);
 }
 
 inline void Free(void* Pointer, size_t AllocationSize)
 {
-	FMemoryStats::UpdateBytesAllocated(-static_cast<ptrdiff_t>(AllocationSize));
 	return GExternAPI.Free(Pointer, AllocationSize);
 }
 
diff --git a/Engine/Source/Runtime/AutoRTFM/Private/MemoryStats.cpp b/Engine/Source/Runtime/AutoRTFM/Private/MemoryStats.cpp
deleted file mode 100644
diff --git a/Engine/Source/Runtime/AutoRTFM/Private/MemoryStats.h b/Engine/Source/Runtime/AutoRTFM/Private/MemoryStats.h
deleted file mode 100644
diff --git a/Engine/Source/Runtime/AutoRTFM/Public/AutoRTFM/CAPI.h b/Engine/Source/Runtime/AutoRTFM/Public/AutoRTFM/CAPI.h
--- a/Engine/Source/Runtime/AutoRTFM/Public/AutoRTFM/CAPI.h
+++ b/Engine/Source/Runtime/AutoRTFM/Public/AutoRTFM/CAPI.h
@@ -102,19 +102,19 @@ extern "C"
 	// Function pointers used by AutoRTFM for heap allocations, etc.
 	typedef struct
 	{
-		// The function used to allocate memory from the heap. Must not be null.
-		// This API should never return null; if the allocator cannot satisfy the request, it must abort the program.
+		// The function used to allocate memory from the heap.
+		// Must not be null.
 		void* (*Allocate)(size_t Size, size_t Alignment);
 
-		// The function used to reallocate memory from the heap. Must not be null.
-		// This API should never return null; if the allocator cannot satisfy the request, it must abort the program.
+		// The function used to reallocate memory from the heap.
+		// Must not be null.
 		void* (*Reallocate)(void* Pointer, size_t Size, size_t Alignment, size_t PreviousSize);
 
-		// The function used to allocate zeroed memory from the heap. Must not be null.
-		// This API should never return null; if the allocator cannot satisfy the request, it must abort the program.
+		// The function used to allocate zeroed memory from the heap.
+		// Must not be null.
 		void* (*AllocateZeroed)(size_t Size, size_t Alignment);
 
-		// The function used to free memory allocated by Allocate(), Reallocate() or AllocateZeroed().
+		// The function used to free memory allocated by Allocate(), Reallocate() and AllocateZeroed().
 		// Must not be null.
 		void (*Free)(void* Pointer, size_t AllocationSize);
 
