[AutoRTFM] Fix link error when compiling AutoRTFMTests in Shipping mode.

Keeping this compile working is useful for shaking out Shipping-specific issues.

#Synced-CL 50109039

diff --git a/Engine/Source/Programs/AutoRTFMTests/AutoRTFMTests.Target.cs b/Engine/Source/Programs/AutoRTFMTests/AutoRTFMTests.Target.cs
--- a/Engine/Source/Programs/AutoRTFMTests/AutoRTFMTests.Target.cs
+++ b/Engine/Source/Programs/AutoRTFMTests/AutoRTFMTests.Target.cs
@@ -86,7 +86,10 @@ public class AutoRTFMTestsTarget : TargetRules
 		GlobalDefinitions.Add("CSV_PROFILER_ENABLE_IN_SHIPPING=1");
 
 		// CallstackTrace tests rely on this being enabled!
-		GlobalDefinitions.Add("UE_CALLSTACK_TRACE_ENABLED=1");
+		if (Target.Configuration != UnrealTargetConfiguration.Shipping)
+		{
+			GlobalDefinitions.Add("UE_CALLSTACK_TRACE_ENABLED=1");
+		}
 
 		// Need to force the heartbeat on so we can test with it!
 		GlobalDefinitions.Add("USE_HITCH_DETECTION=1");
diff --git a/Engine/Source/Programs/AutoRTFMTests/Private/InsightsTests.cpp b/Engine/Source/Programs/AutoRTFMTests/Private/InsightsTests.cpp
--- a/Engine/Source/Programs/AutoRTFMTests/Private/InsightsTests.cpp
+++ b/Engine/Source/Programs/AutoRTFMTests/Private/InsightsTests.cpp
@@ -539,6 +539,7 @@ TEST_CASE("Insights")
 	}
 #endif
 
+#if UE_CALLSTACK_TRACE_ENABLED
 	SECTION("CallstackTrace")
 	{
 		AutoRTFMInsightsDeferer Callstack(CallstackChannel);
@@ -550,4 +551,5 @@ TEST_CASE("Insights")
 			REQUIRE(OutwithTransactionId != CallstackTrace_GetCurrentId());
 		});
 	}
+#endif
 }
