[AutoRTFM] Split OnComplete() into OnComplete() and OnRetry().

AutoRTFM::OnComplete() used to be called for retries as well as no-retry.
This has been split into AutoRTFM::OnRetry() which will be called before the outer transaction is retried, and AutoRTFM::OnComplete() which is called after retrying, once the outer transaction commits or aborts.

This simplifies usage, as callbacks nearly always started with a test for `AutoRTFM::IsRetrying()`.

#Synced-CL 50196324

diff --git a/Engine/Source/Programs/AutoRTFMTests/Private/API.cpp b/Engine/Source/Programs/AutoRTFMTests/Private/API.cpp
--- a/Engine/Source/Programs/AutoRTFMTests/Private/API.cpp
+++ b/Engine/Source/Programs/AutoRTFMTests/Private/API.cpp
@@ -212,39 +212,20 @@ TEST_CASE("API.autortfm_on_commit")
 
 TEST_CASE("API.autortfm_on_abort")
 {
-	// Too hard to get this test working when retrying nested transactions so bail!
-	if (AutoRTFM::ForTheRuntime::ShouldRetryNestedTransactionsToo())
-	{
-		return;
-	}
+	AUTORTFM_SCOPED_DISABLE_RETRY();
 
     bool bOuterTransaction = false;
     bool bInnerTransaction = false;
     bool bInnerTransactionWithAbort = false;
     bool bInnerOpenNest = false;
 
-    REQUIRE(AutoRTFM::ETransactionResult::Committed == AutoRTFM::Transact([&]
+    AutoRTFM::Testing::Commit([&]
     {
-		// If we are retrying transactions, need to reset the test state.
-		AutoRTFM::OnAbort([&]
-		{
-			bOuterTransaction = false;
-			bInnerTransaction = false;
-			bInnerTransactionWithAbort = false;
-			bInnerOpenNest = false;
-		});
-
         autortfm_on_abort([](void* const Arg)
         {
             *static_cast<bool* const>(Arg) = true;
         }, &bOuterTransaction);
 
-        // This should only be modified on the commit!
-        if (bOuterTransaction)
-        {
-            AutoRTFM::AbortTransaction();
-        }
-
         AutoRTFM::Testing::Commit([&]
         {
             autortfm_on_abort([](void* const Arg)
@@ -253,13 +234,7 @@ TEST_CASE("API.autortfm_on_abort")
             }, &bInnerTransaction);
         });
 
-        // This should only be modified on the commit!
-        if (bInnerTransaction)
-        {
-            AutoRTFM::AbortTransaction();
-        }
-
-        AutoRTFM::ETransactionResult Result = AutoRTFM::Transact([&]
+        AutoRTFM::Testing::Abort([&]
         {
             autortfm_on_abort([](void* const Arg)
             {
@@ -269,14 +244,6 @@ TEST_CASE("API.autortfm_on_abort")
             AutoRTFM::AbortTransaction();
         });
 
-        REQUIRE(AutoRTFM::ETransactionResult::AbortedByRequest == Result);
-
-        // OnAbort runs eagerly on inner abort
-        if (!bInnerTransactionWithAbort)
-        {
-            AutoRTFM::AbortTransaction();
-        }
-
         AutoRTFM::Open([&]
         {
             autortfm_on_abort([](void* const Arg)
@@ -284,13 +251,7 @@ TEST_CASE("API.autortfm_on_abort")
                 *static_cast<bool* const>(Arg) = true;
             }, &bInnerOpenNest);
         });
-
-        // This should only be modified on the commit!
-        if (bInnerOpenNest)
-        {
-            AutoRTFM::AbortTransaction();
-        }
-    }));
+    });
 
     REQUIRE(false == bOuterTransaction);
     REQUIRE(false == bInnerTransaction);
@@ -484,6 +445,7 @@ TEST_CASE("API.IsCommittingOrAborting")
 	bool bInComplete = false;
 	bool bInRetryAbort = false;
 	bool bInRetryComplete = false;
+	bool bInRetryRetry = false;
 
 	AutoRTFM::Testing::Commit([&]
 	{
@@ -513,10 +475,12 @@ TEST_CASE("API.IsCommittingOrAborting")
 		{
 			bInOpenNest = AutoRTFM::IsCommittingOrAborting();
 
-			REQUIRE(AutoRTFM::EContextStatus::OnTrack == AutoRTFM::Close([&]
+			AutoRTFM::EContextStatus Status = AutoRTFM::Close([&]
 			{
 				bInClosedNestInOpenNest = AutoRTFM::IsCommittingOrAborting();
-			}));
+			});
+
+			REQUIRE(AutoRTFM::EContextStatus::OnTrack == Status);
 		});
 	});
 
@@ -530,9 +494,13 @@ TEST_CASE("API.IsCommittingOrAborting")
 			{
 				bInRetryAbort = AutoRTFM::IsCommittingOrAborting();
 			});
-			AutoRTFM::OnComplete([&]
+			AutoRTFM::OnAbort([&]
 			{
 				bInRetryComplete = AutoRTFM::IsCommittingOrAborting();
+			});
+			AutoRTFM::OnRetry([&]
+			{
+				bInRetryRetry = AutoRTFM::IsCommittingOrAborting();
 				bAlreadyRetried = true;
 			});
 
@@ -548,86 +516,231 @@ TEST_CASE("API.IsCommittingOrAborting")
 	REQUIRE(true == bInComplete);
 	REQUIRE(true == bInRetryAbort);
 	REQUIRE(true == bInRetryComplete);
+	REQUIRE(true == bInRetryRetry);
 }
 
-
-TEST_CASE("API.IsRetrying")
+TEST_CASE("API.OnRetry")
 {
-	REQUIRE(false == AutoRTFM::IsRetrying());
+	AUTORTFM_SCOPED_DISABLE_RETRY();
 
-	// Set to the opposite of what we expect at the end of function.
-	bool bInTransaction = true;
-	bool bInOpenNest = true;
-	bool bInClosedNestInOpenNest = true;
-	bool bInAbort = true;
-	bool bInCommit = true;
-	bool bInComplete = true;
-	bool bInRetryAbort = false;
-	bool bInRetryComplete = false;
+	enum EEvent
+	{
+		OnStartTransaction,
+		OnEndTransaction,
+		OnAbort1,
+		OnAbort2,
+		OnCommit1,
+		OnCommit2,
+		OnComplete1,
+		OnComplete2,
+		OnRetry1,
+		OnRetry2,
+		OnAbortDestructor1,
+		OnAbortDestructor2,
+		OnCommitDestructor1,
+		OnCommitDestructor2,
+		OnCompleteDestructor1,
+		OnCompleteDestructor2,
+		OnRetryDestructor1,
+		OnRetryDestructor2,
+	};
 
-	AutoRTFM::Testing::Commit([&]
+	enum EMode
 	{
-		bInTransaction = AutoRTFM::IsRetrying();
+		Commit,
+		Abort,
+		Retry,
+	};
 
-		AutoRTFM::Testing::Abort([&]
+	auto Run = [] (EMode Mode) -> std::vector<EEvent>
+	{
+		std::vector<EEvent> Events;
+		auto AddEvent = [&] AUTORTFM_OPEN (EEvent Event){
+			Events.push_back(Event);
+		};
+		struct AUTORTFM_ENABLE FAddEventOnDestruction
 		{
-			AutoRTFM::OnAbort([&]
+			std::vector<EEvent>& Events;
+			EEvent Event;
+			mutable bool bMovedOrCopied = false;
+			FAddEventOnDestruction(std::vector<EEvent>& Events, EEvent Event) : Events{Events}, Event{Event} {}
+			FAddEventOnDestruction(const FAddEventOnDestruction& Other) : Events{Other.Events}, Event{Other.Event}
 			{
-				bInAbort = AutoRTFM::IsRetrying();
-			});
+				Other.bMovedOrCopied = true;
+			}
+			FAddEventOnDestruction(FAddEventOnDestruction&& Other) : Events{Other.Events}, Event{Other.Event}
+			{
+				Other.bMovedOrCopied = true;
+			}
+			FAddEventOnDestruction& operator = (const FAddEventOnDestruction&) = delete;
+			FAddEventOnDestruction& operator = (FAddEventOnDestruction&&) = delete;
 
-			AutoRTFM::AbortTransaction();
-		});
+			~FAddEventOnDestruction()
+			{
+				if (!bMovedOrCopied)
+				{
+					AutoRTFM::Open([&] { Events.push_back(Event); });
+				}
+			}
+		};
 
-		AutoRTFM::OnCommit([&]
+		bool bHasRetried = false;
+	
+		AutoRTFM::ETransactionResult Result = AutoRTFM::Transact([&]
 		{
-			bInCommit = AutoRTFM::IsRetrying();
-		});
+			AddEvent(OnStartTransaction);
 
-		AutoRTFM::OnComplete([&]
-		{
-			bInComplete = AutoRTFM::IsRetrying();
-		});
+			AutoRTFM::OnAbort([&, Dtor = FAddEventOnDestruction{Events, OnAbortDestructor1}]
+			{
+				AddEvent(OnAbort1);
+			});
 
-		AutoRTFM::Open([&]
-		{
-			bInOpenNest = AutoRTFM::IsRetrying();
+			AutoRTFM::OnCommit([&, Dtor = FAddEventOnDestruction{Events, OnCommitDestructor1}]
+			{
+				AddEvent(OnCommit1);
+			});
 
-			REQUIRE(AutoRTFM::EContextStatus::OnTrack == AutoRTFM::Close([&]
+			AutoRTFM::OnComplete([&, Dtor = FAddEventOnDestruction{Events, OnCompleteDestructor1}]
 			{
-				bInClosedNestInOpenNest = AutoRTFM::IsRetrying();
-			}));
-		});
-	});
+				AddEvent(OnComplete1);
+			});
 
-	bool bAlreadyRetried = false;
+			AutoRTFM::OnRetry([&, Dtor = FAddEventOnDestruction{Events, OnRetryDestructor1}]
+			{
+				AddEvent(OnRetry1);
+				bHasRetried = true;
+			});
 
-	AutoRTFM::Testing::Commit([&]
-	{
-		if (!bAlreadyRetried)
-		{
-			AutoRTFM::OnAbort([&]
+			AutoRTFM::OnAbort([&, Dtor = FAddEventOnDestruction{Events, OnAbortDestructor2}]
 			{
-				bInRetryAbort = AutoRTFM::IsRetrying();
+				AddEvent(OnAbort2);
 			});
-			AutoRTFM::OnComplete([&]
+
+			AutoRTFM::OnCommit([&, Dtor = FAddEventOnDestruction{Events, OnCommitDestructor2}]
 			{
-				bInRetryComplete = AutoRTFM::IsRetrying();
-				bAlreadyRetried = true;
+				AddEvent(OnCommit2);
 			});
 
-			AutoRTFM::CascadingRetryTransaction();
+			AutoRTFM::OnComplete([&, Dtor = FAddEventOnDestruction{Events, OnCompleteDestructor2}]
+			{
+				AddEvent(OnComplete2);
+			});
+
+			AutoRTFM::OnRetry([&, Dtor = FAddEventOnDestruction{Events, OnRetryDestructor2}]
+			{
+				AddEvent(OnRetry2);
+				bHasRetried = true;
+			});
+
+			switch (Mode)
+			{
+				case Commit:
+					break;
+				case Abort:
+					AutoRTFM::AbortTransaction();
+					break;
+				case Retry:
+					if (!bHasRetried)
+					{
+						AutoRTFM::CascadingRetryTransaction();
+					}
+					break;
+			}
+			
+			AddEvent(OnEndTransaction);
+		});
+		
+		switch (Mode)
+		{
+			case Commit:
+				REQUIRE(Result == AutoRTFM::ETransactionResult::Committed);
+				break;
+			case Abort:
+				REQUIRE(Result == AutoRTFM::ETransactionResult::AbortedByRequest);
+				break;
+			case Retry:
+				REQUIRE(Result == AutoRTFM::ETransactionResult::Committed);
+				break;
 		}
-	});
+		
+		return Events;
+	};
 
-	REQUIRE(false == bInTransaction);
-	REQUIRE(false == bInOpenNest);
-	REQUIRE(false == bInClosedNestInOpenNest);
-	REQUIRE(false == bInAbort);
-	REQUIRE(false == bInCommit);
-	REQUIRE(false == bInComplete);
-	REQUIRE(true == bInRetryAbort);
-	REQUIRE(true == bInRetryComplete);
+	SECTION("Commit")
+	{
+		std::vector<EEvent> Events = Run(Commit);
+		REQUIRE(Events == std::vector<EEvent>{
+			OnStartTransaction,
+			OnEndTransaction,
+			OnAbortDestructor1,
+			OnAbortDestructor2,
+			OnCommit1,
+			OnCommitDestructor1,
+			OnCommit2,
+			OnCommitDestructor2,
+			OnRetryDestructor1,
+			OnRetryDestructor2,
+			OnComplete1,
+			OnCompleteDestructor1,
+			OnComplete2,
+			OnCompleteDestructor2,
+		});
+	}
+
+	SECTION("Abort")
+	{
+		std::vector<EEvent> Events = Run(Abort);
+		REQUIRE(Events == std::vector<EEvent>{
+			OnStartTransaction,
+			OnCommitDestructor1,
+			OnCommitDestructor2,
+			OnAbort2,
+			OnAbortDestructor2,
+			OnAbort1,
+			OnAbortDestructor1,
+			OnRetryDestructor1,
+			OnRetryDestructor2,
+			OnComplete1,
+			OnCompleteDestructor1,
+			OnComplete2,
+			OnCompleteDestructor2,
+		});
+	}
+
+	SECTION("Retry")
+	{
+		std::vector<EEvent> Events = Run(Retry);
+		REQUIRE(Events == std::vector<EEvent>{
+			OnStartTransaction,
+			OnCommitDestructor1,
+			OnCommitDestructor2,
+			OnAbort2,
+			OnAbortDestructor2,
+			OnAbort1,
+			OnAbortDestructor1,
+			OnCompleteDestructor1,
+			OnCompleteDestructor2,
+			OnRetry1,
+			OnRetryDestructor1,
+			OnRetry2,
+			OnRetryDestructor2,
+			// Retry
+			OnStartTransaction,
+			OnEndTransaction,
+			OnAbortDestructor1,
+			OnAbortDestructor2,
+			OnCommit1,
+			OnCommitDestructor1,
+			OnCommit2,
+			OnCommitDestructor2,
+			OnRetryDestructor1,
+			OnRetryDestructor2,
+			OnComplete1,
+			OnCompleteDestructor1,
+			OnComplete2,
+			OnCompleteDestructor2,
+		});
+	}
 }
 
 TEST_CASE("API.Transact")
@@ -1051,11 +1164,7 @@ TEST_CASE("API.OnCommitMacro_WithAbort")
 
 TEST_CASE("API.OnAbort")
 {
-	// Too hard to get this test working when retrying nested transactions so bail!
-	if (AutoRTFM::ForTheRuntime::ShouldRetryNestedTransactionsToo())
-	{
-		return;
-	}
+	AUTORTFM_SCOPED_DISABLE_RETRY();
 
     bool bOuterTransaction = false;
     bool bInnerTransaction = false;
@@ -1139,11 +1248,7 @@ TEST_CASE("API.OnAbort")
 
 TEST_CASE("API.OnPreAbort")
 {
-	// Too hard to get this test working when retrying nested transactions so bail!
-	if (AutoRTFM::ForTheRuntime::ShouldRetryNestedTransactionsToo())
-	{
-		return;
-	}
+	AUTORTFM_SCOPED_DISABLE_RETRY();
 
     bool bOuterTransaction = false;
     bool bInnerTransaction = false;
@@ -1424,86 +1529,90 @@ TEST_CASE("API.OnCompleteIgnoredOutsideTransaction")
 
 TEST_CASE("API.CascadingRetryTransaction")
 {
+	AUTORTFM_SCOPED_DISABLE_RETRY();
+
+	size_t CallCount = 0;
+	bool bTransactionalWrite = false;
+
 	SECTION("Non-nested committed transaction")
 	{
-		bool bFirst = true;
 		AutoRTFM::Testing::Commit([&]
 		{
-			if (bFirst)
+			bTransactionalWrite = true;
+			if (CallCount == 0)
 			{
-				bFirst = false;
-				AutoRTFM::OnComplete([&]
+				AutoRTFM::OnRetry([&]
 				{
 					// The write above would have been undone.
-					REQUIRE(bFirst);
-					bFirst = false;
+					REQUIRE(!bTransactionalWrite);
+					CallCount++;
 				});
 				AutoRTFM::CascadingRetryTransaction();
 			}
 		});
-		REQUIRE(!bFirst);
+		REQUIRE(bTransactionalWrite);
+		REQUIRE(CallCount == 1);
 	}
 
 	SECTION("Non-nested aborted transaction")
 	{
-		bool bFirst = true;
 		AutoRTFM::Testing::Abort([&]
 		{
-			if (bFirst)
+			bTransactionalWrite = true;
+			if (CallCount == 0)
 			{
-				bFirst = false;
-				AutoRTFM::OnComplete([&]
+				AutoRTFM::OnRetry([&]
 				{
 					// The write above would have been undone.
-					REQUIRE(bFirst);
-					bFirst = false;
+					REQUIRE(!bTransactionalWrite);
+					CallCount++;
 				});
 				AutoRTFM::CascadingRetryTransaction();
 			}
 
 			AutoRTFM::AbortTransaction();
 		});
-		REQUIRE(!bFirst);
+		REQUIRE(!bTransactionalWrite);
+		REQUIRE(CallCount == 1);
 	}
 
 	SECTION("Nested committed transaction")
 	{
-		bool bFirst = true;
 		AutoRTFM::Testing::Commit([&]
 		{
 			AutoRTFM::Testing::Commit([&]
 			{
-				if (bFirst)
+				bTransactionalWrite = true;
+				if (CallCount == 0)
 				{
-					bFirst = false;
-					AutoRTFM::OnComplete([&]
+					AutoRTFM::OnRetry([&]
 					{
 						// The write above would have been undone.
-						REQUIRE(bFirst);
-						bFirst = false;
+						REQUIRE(!bTransactionalWrite);
+						CallCount++;
 					});
 					AutoRTFM::CascadingRetryTransaction();
 				}
 			});
 		});
-		REQUIRE(!bFirst);
+		REQUIRE(bTransactionalWrite);
+		REQUIRE(CallCount == 1);
 	}
 
 	SECTION("Nested aborted transaction")
 	{
-		bool bFirst = true;
 		AutoRTFM::Testing::Commit([&]
 		{
 			AutoRTFM::Testing::Abort([&]
 			{
-				if (bFirst)
+				bTransactionalWrite = true;
+				if (CallCount == 0)
 				{
-					bFirst = false;
-					AutoRTFM::OnComplete([&]
+					AutoRTFM::OnRetry([&]
 					{
 						// The write above would have been undone.
-						REQUIRE(bFirst);
-						bFirst = false;
+						REQUIRE(!bTransactionalWrite);
+						CallCount++;
 					});
 					AutoRTFM::CascadingRetryTransaction();
 				}
@@ -1511,43 +1620,29 @@ TEST_CASE("API.CascadingRetryTransaction")
 				AutoRTFM::AbortTransaction();
 			});
 		});
-		REQUIRE(!bFirst);
+		REQUIRE(!bTransactionalWrite);
+		REQUIRE(CallCount == 1);
 	}
 
-	SECTION("IsTransactional is false during the retry")
+	SECTION("Status query during retry")
 	{
-		bool bFirst = true;
 		AutoRTFM::Testing::Commit([&]
 		{
-			if (bFirst)
+			if (CallCount == 0)
 			{
-				AutoRTFM::OnComplete([&]
+				AutoRTFM::OnRetry([&]
 				{
 					REQUIRE(!AutoRTFM::IsTransactional());
-					bFirst = false;
-				});
-				AutoRTFM::CascadingRetryTransaction();
-			}
-		});
-		REQUIRE(!bFirst);
-	}
-
-	SECTION("IsClosed is false during the retry")
-	{
-		bool bFirst = true;
-		AutoRTFM::Testing::Commit([&]
-		{
-			if (bFirst)
-			{
-				AutoRTFM::OnComplete([&]
-				{
 					REQUIRE(!AutoRTFM::IsClosed());
-					bFirst = false;
+					REQUIRE(!AutoRTFM::IsCommitting());
+					REQUIRE(AutoRTFM::IsRetrying());
+					REQUIRE(AutoRTFM::IsCommittingOrAborting());
+					CallCount++;
 				});
 				AutoRTFM::CascadingRetryTransaction();
 			}
 		});
-		REQUIRE(!bFirst);
+		REQUIRE(CallCount == 1);
 	}
 }
 
diff --git a/Engine/Source/Programs/AutoRTFMTests/Private/Abort.cpp b/Engine/Source/Programs/AutoRTFMTests/Private/Abort.cpp
--- a/Engine/Source/Programs/AutoRTFMTests/Private/Abort.cpp
+++ b/Engine/Source/Programs/AutoRTFMTests/Private/Abort.cpp
@@ -191,14 +191,17 @@ TEST_CASE("Abort.TransactionInAbortedOnComplete")
 	});
 }
 
-TEST_CASE("Abort.TransactionInCommittedOnComplete")
+TEST_CASE("Abort.TransactionInOnComplete")
 {
-	bool bGotAbortedDuringCommit = false;
+	AUTORTFM_SCOPED_DISABLE_RETRY();
 
+	size_t CallCount = 0;
 	AutoRTFM::Testing::Commit([&]
 	{
 		AutoRTFM::OnComplete([&]
 		{
+			CallCount++;
+
 			bool bDidSomething = false;
 
 			AutoRTFM::ETransactionResult InnerResult = AutoRTFM::Transact([&]
@@ -206,19 +209,42 @@ TEST_CASE("Abort.TransactionInCommittedOnComplete")
 				bDidSomething = true;
 			});
 
-			// The test harness will abort this transaction and retry it when `ShouldRetryNonNestedTransactions`
-			// is set, which unfortunately muddles the test result here since we'd normally expect "transact during
-			// commit" here, not "transact during abort."
-			REQUIRE((InnerResult == AutoRTFM::ETransactionResult::AbortedByTransactDuringCommit ||
-			         InnerResult == AutoRTFM::ETransactionResult::AbortedByTransactDuringAbort));
+			REQUIRE(!bDidSomething);
+			REQUIRE(InnerResult == AutoRTFM::ETransactionResult::AbortedByTransactDuringCommit);
+		});
+	});
+	REQUIRE(CallCount == 1);
+}
+
+TEST_CASE("Abort.TransactionInOnRetry")
+{
+	AUTORTFM_SCOPED_DISABLE_RETRY();
+
+	size_t CallCount = 0;
+	AutoRTFM::Testing::Commit([&]
+	{
+		AutoRTFM::OnRetry([&]
+		{
+			CallCount++;
+
+			bool bDidSomething = false;
+
+			AutoRTFM::ETransactionResult InnerResult = AutoRTFM::Transact([&]
+			{
+				bDidSomething = true;
+			});
 
-			bGotAbortedDuringCommit |= (InnerResult == AutoRTFM::ETransactionResult::AbortedByTransactDuringCommit);
 			REQUIRE(!bDidSomething);
+			REQUIRE(InnerResult == AutoRTFM::ETransactionResult::AbortedByTransactDuringAbort);
 		});
+		
+		if (CallCount == 0)
+		{
+			AutoRTFM::CascadingRetryTransaction();
+		}
 	});
 
-	// The test only passes if we actually got aborted during commit, so verify that this really happened.
-	REQUIRE(bGotAbortedDuringCommit);
+	REQUIRE(CallCount == 1);
 }
 
 TEST_CASE("Abort.AbortInInnerTransaction")
@@ -438,16 +464,18 @@ TEST_CASE("Abort.CascadeWithOnCompleteThroughOpen")
 				{
 					AutoRTFM::OnComplete([&]
 					{
-						REQUIRE(!bTouched);  // Rollbacks must occur before the post-abort callback is invoked.
+						REQUIRE(!bTouched);  // Rollbacks must occur before the complete callback is invoked.
 						bOnCompleteCallbackWasExecuted = true;
 					});
 					AutoRTFM::CascadingAbortTransaction();
 				});
+
+				FAIL("AutoRTFM::Transact() failed to throw after a cascading abort");
 			});
 
 			REQUIRE(AutoRTFM::EContextStatus::AbortedByCascadingAbort == Status);
 		});
-		FAIL("AutoRTFM::Open() failed to throw after an abort");
+		FAIL("AutoRTFM::Open() failed to throw after a cascading abort");
 	});
 
 	REQUIRE(AutoRTFM::ETransactionResult::AbortedByCascade == Result);
@@ -493,6 +521,7 @@ TEST_CASE("Abort.CascadeWithOnCompleteThroughFunctionAttribute")
 					AutoRTFM::CascadingAbortTransaction();
 				});
 			});
+			REQUIRE(AutoRTFM::EContextStatus::AbortedByCascadingAbort == Status);
 		}
 
 		bool bTouched = false;
@@ -1124,7 +1153,7 @@ TEST_CASE("Abort.StackWriteToOuterOpen")
 	bool WritesUndone = true;
 	bool Success = false;
 
-	const AutoRTFM::ETransactionResult TransactionResult = AutoRTFM::Transact([&]
+	AutoRTFM::Testing::Abort([&]
 	{
 		AutoRTFM::Open([&]
 		{
@@ -1185,7 +1214,6 @@ TEST_CASE("Abort.StackWriteToOuterOpen")
 	});
 
 	REQUIRE(Success);
-	REQUIRE(AutoRTFM::ETransactionResult::AbortedByRequest == TransactionResult);
 }
 
 TEST_CASE("Abort.PushOnAbortOrdering")
diff --git a/Engine/Source/Programs/AutoRTFMTests/Private/UObjectTests.cpp b/Engine/Source/Programs/AutoRTFMTests/Private/UObjectTests.cpp
--- a/Engine/Source/Programs/AutoRTFMTests/Private/UObjectTests.cpp
+++ b/Engine/Source/Programs/AutoRTFMTests/Private/UObjectTests.cpp
@@ -952,7 +952,7 @@ TEST_CASE("UObject.ReleaseRef")
 
 			if (bFirst)
 			{
-				AutoRTFM::OnComplete([&] 
+				AutoRTFM::OnRetry([&] 
 				{ 
 					bFirst = false; 
 				});
diff --git a/Engine/Source/Programs/AutoRTFMTests/Private/UPackageTests.cpp b/Engine/Source/Programs/AutoRTFMTests/Private/UPackageTests.cpp
--- a/Engine/Source/Programs/AutoRTFMTests/Private/UPackageTests.cpp
+++ b/Engine/Source/Programs/AutoRTFMTests/Private/UPackageTests.cpp
@@ -648,7 +648,7 @@ TEST_CASE("UPackage.AsyncLoading")
 				return;
 			}
 
-			AutoRTFM::OnComplete([&] { bFirst = false; });
+			AutoRTFM::OnRetry([&] { bFirst = false; });
 
 			FlushAsyncLoading(LoadPackageAsync(FPackagePath::FromPackageNameUnchecked(*Name), CustomName, Delegate));
 		});
diff --git a/Engine/Source/Runtime/AutoRTFM/Private/API.cpp b/Engine/Source/Runtime/AutoRTFM/Private/API.cpp
--- a/Engine/Source/Runtime/AutoRTFM/Private/API.cpp
+++ b/Engine/Source/Runtime/AutoRTFM/Private/API.cpp
@@ -595,6 +595,13 @@ AUTORTFM_INTERNAL void RTFM_OnCompleteInternal(TTask<void()>&& Work)
 	Context->GetCurrentTransaction()->DeferUntilComplete(std::move(Work));
 }
 
+AUTORTFM_INTERNAL void RTFM_OnRetryInternal(TTask<void()>&& Work)
+{
+	FContext* Context = FContext::Get();
+	AUTORTFM_ASSERT(Context->GetStatus() == EContextStatus::OnTrack);
+	Context->GetCurrentTransaction()->DeferUntilRetry(std::move(Work));
+}
+
 AUTORTFM_INTERNAL void RTFM_PushOnCommitHandlerInternal(const void* Key, TTask<void()>&& Work)
 {
 	FContext* Context = FContext::Get();
@@ -1070,6 +1077,14 @@ void ForTheRuntime::OnCompleteInternal(TTask<void()>&& Work)
 	}
 }
 
+void ForTheRuntime::OnRetryInternal(TTask<void()>&& Work)
+{
+	if (autortfm_is_closed())
+	{
+		return RTFM_OnRetryInternal(std::move(Work));
+	}
+}
+
 void ForTheRuntime::PushOnCommitHandlerInternal(const void* Key, TTask<void()>&& Work)
 {
 	if (autortfm_is_closed())
diff --git a/Engine/Source/Runtime/AutoRTFM/Private/Transaction.cpp b/Engine/Source/Runtime/AutoRTFM/Private/Transaction.cpp
--- a/Engine/Source/Runtime/AutoRTFM/Private/Transaction.cpp
+++ b/Engine/Source/Runtime/AutoRTFM/Private/Transaction.cpp
@@ -69,8 +69,9 @@ void FTransaction::Initialize(FTransaction* InParent, bool bInIsStackScoped, FSt
 	}
 	else
 	{
-		// Root-level transactions require a completion task list.
+		// Root-level transactions require completion and retry task lists.
 		CompletionTasks.emplace(Context->GetTaskPool());
+		RetryTasks.emplace(Context->GetTaskPool());
 	}
 
 	Stats.Collect<EStatsKind::AverageTransactionDepth>(StatDepth);
@@ -87,7 +88,11 @@ void FTransaction::Reset()
 		CompletionTasks->Reset();
 		CompletionTasks = std::nullopt;
 	}
-
+	if (RetryTasks != std::nullopt)
+	{
+		RetryTasks->Reset();
+		RetryTasks = std::nullopt;
+	}
 	HitSet.Reset();
 	NewMemoryTracker.Reset();
 	WriteLog.Reset();
@@ -106,7 +111,8 @@ bool FTransaction::IsFresh() const
 {
 	return HitSet.IsEmpty() && NewMemoryTracker.IsEmpty() && WriteLog.IsEmpty() && CommitTasks.IsEmpty() && PreAbortTasks.IsEmpty()
 		&& AbortTasks.IsEmpty() && (CompletionTasks == std::nullopt || CompletionTasks->IsEmpty())
-		&& DeferredPopOnCommitHandlers.IsEmpty() && DeferredPopOnAbortHandlers.IsEmpty() && DeferredPopAllOnCommitHandlers.IsEmpty()
+		&& (RetryTasks == std::nullopt || RetryTasks->IsEmpty()) && DeferredPopOnCommitHandlers.IsEmpty()
+		&& DeferredPopOnAbortHandlers.IsEmpty() && DeferredPopAllOnCommitHandlers.IsEmpty()
 		&& DeferredPopAllOnAbortHandlers.IsEmpty()
 #if AUTORTFM_SANITIZER
 		&& ClosedWrites.IsEmpty()
@@ -143,7 +149,16 @@ void FTransaction::AbortWithoutThrowing()
 	}
 	else
 	{
-		CompletionTasks->RemoveEachForward([&](TTask<void()>& Task) { Task(); });
+		if (Context->IsRetrying())
+		{
+			CompletionTasks->Reset();
+			RetryTasks->RemoveEachForward([&](TTask<void()>& Task) { Task(); });
+		}
+		else
+		{
+			RetryTasks->Reset();
+			CompletionTasks->RemoveEachForward([&](TTask<void()>& Task) { Task(); });
+		}
 
 		AUTORTFM_ASSERT(Context->IsAborting());
 	}
@@ -208,6 +223,7 @@ void FTransaction::CommitNested()
 {
 	AUTORTFM_ASSERT(Parent);
 	AUTORTFM_ASSERT(CompletionTasks == std::nullopt);
+	AUTORTFM_ASSERT(RetryTasks == std::nullopt);
 
 	// Disable the sanitizer for the scope of this function as we may allocate
 	// memory, which may call autortfm_sanitizer_open_write().
@@ -316,6 +332,7 @@ bool FTransaction::AttemptToCommitOuterNest()
 
 	CommitTasks.RemoveEachForward([](TTask<void()>& Task) { Task(); });
 
+	RetryTasks->Reset();
 	CompletionTasks->RemoveEachForward([](TTask<void()>& Task) { Task(); });
 
 	return true;
diff --git a/Engine/Source/Runtime/AutoRTFM/Private/Transaction.h b/Engine/Source/Runtime/AutoRTFM/Private/Transaction.h
--- a/Engine/Source/Runtime/AutoRTFM/Private/Transaction.h
+++ b/Engine/Source/Runtime/AutoRTFM/Private/Transaction.h
@@ -79,6 +79,7 @@ public:
 	void DeferUntilPreAbort(TTask<void()>&&);
 	void DeferUntilAbort(TTask<void()>&&);
 	void DeferUntilComplete(TTask<void()>&&);
+	void DeferUntilRetry(TTask<void()>&&);
 	void PushDeferUntilCommitHandler(const void* Key, TTask<void()>&&);
 	void PopDeferUntilCommitHandler(const void* Key);
 	void PopAllDeferUntilCommitHandlers(const void* Key);
@@ -180,6 +181,9 @@ private:
 	// Completion tasks run at the end of the transaction, in forward order.
 	std::optional<FTaskArray> CompletionTasks;
 
+	// Retry tasks run before retry of the transaction, in forward order.
+	std::optional<FTaskArray> RetryTasks;
+
 	// If a call to `PopOnCommitHandler` could not find a commit to pop, it is deferred
 	// and tried again on the parent transaction.
 	TStack<const void*, 8> DeferredPopOnCommitHandlers;
diff --git a/Engine/Source/Runtime/AutoRTFM/Private/TransactionInlines.h b/Engine/Source/Runtime/AutoRTFM/Private/TransactionInlines.h
--- a/Engine/Source/Runtime/AutoRTFM/Private/TransactionInlines.h
+++ b/Engine/Source/Runtime/AutoRTFM/Private/TransactionInlines.h
@@ -272,6 +272,22 @@ UE_AUTORTFM_FORCEINLINE void FTransaction::DeferUntilComplete(TTask<void()>&& Ca
 	Transaction->CompletionTasks->Add(std::move(Copy));
 }
 
+UE_AUTORTFM_FORCEINLINE void FTransaction::DeferUntilRetry(TTask<void()>&& Callback)
+{
+	// Retry tasks are always stored directly on the root level of a transaction.
+	FTransaction* Transaction = this;
+	while (FTransaction* Above = Transaction->Parent)
+	{
+		Transaction = Above;
+	}
+
+	// We explicitly must copy the function here because the original was allocated
+	// within a transactional context, and thus the memory is allocating under
+	// transactionalized conditions. By copying, we create an open copy of the callback.
+	TTask<void()> Copy(Callback);
+	Transaction->RetryTasks->Add(std::move(Copy));
+}
+
 UE_AUTORTFM_FORCEINLINE void FTransaction::PushDeferUntilCommitHandler(const void* Key, TTask<void()>&& Callback)
 {
 	// We explicitly must copy the function here because the original was allocated
diff --git a/Engine/Source/Runtime/AutoRTFM/Public/AutoRTFM.h b/Engine/Source/Runtime/AutoRTFM/Public/AutoRTFM.h
--- a/Engine/Source/Runtime/AutoRTFM/Public/AutoRTFM.h
+++ b/Engine/Source/Runtime/AutoRTFM/Public/AutoRTFM.h
@@ -89,7 +89,7 @@ enum class ETransactionResult
 	// current transaction was being committed.
 	AbortedByTransactDuringCommit = autortfm_aborted_by_transact_during_commit,
 
-	// The transaction aborted because a new transaction nest was attempted (via OnAbort or OnComplete) while the
+	// The transaction aborted because a new transaction nest was attempted (via OnAbort or OnRetry) while the
 	// current transaction was being aborted.
 	AbortedByTransactDuringAbort = autortfm_aborted_by_transact_during_abort,
 
@@ -236,6 +236,7 @@ UE_AUTORTFM_API void OnCommitInternal(TTask<void()>&& Work);
 UE_AUTORTFM_API void OnPreAbortInternal(TTask<void()>&& Work);
 UE_AUTORTFM_API void OnAbortInternal(TTask<void()>&& Work);
 UE_AUTORTFM_API void OnCompleteInternal(TTask<void()>&& Work);
+UE_AUTORTFM_API void OnRetryInternal(TTask<void()>&& Work);
 UE_AUTORTFM_API void PushOnCommitHandlerInternal(const void* Key, TTask<void()>&& Work);
 UE_AUTORTFM_API void PopOnCommitHandlerInternal(const void* Key);
 UE_AUTORTFM_API void PopAllOnCommitHandlersInternal(const void* Key);
@@ -329,14 +330,14 @@ UE_AUTORTFM_CRITICAL_INLINE bool IsCommitting()
 }
 
 // Tells us if we are currently committing or aborting a transaction. This will return true
-// when inside an on-abort, on-commit or on-complete handler.
+// when inside an on-abort, on-commit, on-complete or on-retry handler.
 UE_AUTORTFM_CRITICAL_INLINE bool IsCommittingOrAborting()
 {
 	return autortfm_is_committing_or_aborting();
 }
 
 // Tells us if we are currently retrying a transaction. This will return true when inside an
-// on-abort or on-complete handler that was triggered by CascadingRetryTransaction.
+// on-abort or on-retry handler that was triggered by CascadingRetryTransaction.
 UE_AUTORTFM_CRITICAL_INLINE bool IsRetrying()
 {
 	return autortfm_is_retrying();
@@ -589,7 +590,7 @@ template <typename FunctorType>
 }
 
 // Force a transaction nest to be retried. Once the transaction has aborted,
-// all on-complete handlers will be called before retrying the transaction.
+// all on-retry handlers will be called before retrying the transaction.
 // This is an expensive operation and should thus be used with extreme caution.
 // If this is called outside of a transaction, the call is ignored.
 #if UE_AUTORTFM
@@ -669,10 +670,10 @@ UE_AUTORTFM_CRITICAL_INLINE void OnAbort(FunctorType&&)
 #endif
 
 #if UE_AUTORTFM
-// Have some work happen when the transaction completes, whether that transaction is committed or aborted.
-// In a nested transaction, the work is deferred until the very end of the outermost nest, right before
-// the end of the transaction.
-// The worklist is run in FIFO order, after the OnCommit or OnAbort worklist has finished.
+// Enqueues a task to run after the outermost transaction commits or aborts,
+// after all other OnCommit() or OnAbort() tasks are run, and after any retries
+// of the outermost transaction.
+// The worklist is run in FIFO order.
 // If this is called outside a transaction or from an open nest, then the work is ignored.
 template <typename FunctorType>
 UE_AUTORTFM_CRITICAL_INLINE void OnComplete(AUTORTFM_IMPLICIT_DISABLE FunctorType&& Work)
@@ -689,6 +690,26 @@ UE_AUTORTFM_CRITICAL_INLINE void OnComplete(FunctorType&&)
 }
 #endif
 
+#if UE_AUTORTFM
+// Enqueues a task to run after the outermost transaction retries due to
+// CascadingRetryTransaction(), after all other OnAbort() tasks are run.
+// The worklist is run in FIFO order.
+// If this is called outside a transaction or from an open nest, then the work is ignored.
+template <typename FunctorType>
+UE_AUTORTFM_CRITICAL_INLINE void OnRetry(AUTORTFM_IMPLICIT_DISABLE FunctorType&& Work)
+{
+	if (autortfm_is_closed())
+	{
+		ForTheRuntime::OnRetryInternal(std::forward<FunctorType>(Work));
+	}
+}
+#else
+template <typename FunctorType>
+UE_AUTORTFM_CRITICAL_INLINE void OnRetry(FunctorType&&)
+{
+}
+#endif
+
 #if UE_AUTORTFM
 // Register a handler for transaction commit. Takes a key parameter so that
 // the handler can be unregistered (see `PopOnCommitHandler`). This is useful
diff --git a/Engine/Source/Runtime/CoreUObject/Private/Serialization/TransactionallySafeAsyncLoading.cpp b/Engine/Source/Runtime/CoreUObject/Private/Serialization/TransactionallySafeAsyncLoading.cpp
--- a/Engine/Source/Runtime/CoreUObject/Private/Serialization/TransactionallySafeAsyncLoading.cpp
+++ b/Engine/Source/Runtime/CoreUObject/Private/Serialization/TransactionallySafeAsyncLoading.cpp
@@ -189,7 +189,7 @@ public:
 							AutoRTFM::OnCommit([this] { DrainTransactionallyDeferredLoadPackages(); });
 							AutoRTFM::OnComplete([this]
 							{
-								if (!AutoRTFM::IsRetrying() && !TransactionallyDeferredLoadPackages.IsEmpty())
+								if (!TransactionallyDeferredLoadPackages.IsEmpty())
 								{
 									for (const FTransactionallyDeferredLoadPackagePayload& Payload : TransactionallyDeferredLoadPackages)
 									{
@@ -350,9 +350,8 @@ public:
 
 					const AutoRTFM::EContextStatus Status = AutoRTFM::Close([&]
 					{
-						AutoRTFM::OnComplete([this, RequestIdsClone = ::MoveTemp(RequestIdsClone)]() mutable
+						AutoRTFM::OnRetry([this, RequestIdsClone = ::MoveTemp(RequestIdsClone)]() mutable
 						{
-							check(AutoRTFM::IsRetrying());
 							UE_LOG(LogStreaming, Display, TEXT("A call to `FlushLoading` that is flushing non-cached packages is causing Verse to abort and retry."));
 							DrainTransactionallyDeferredLoadPackages(&RequestIdsClone, true);
 							UpdateMaxFlushedRequestId(RequestIdsClone);
diff --git a/Engine/Source/Runtime/Engine/Private/Components/SkeletalMeshComponent.cpp b/Engine/Source/Runtime/Engine/Private/Components/SkeletalMeshComponent.cpp
--- a/Engine/Source/Runtime/Engine/Private/Components/SkeletalMeshComponent.cpp
+++ b/Engine/Source/Runtime/Engine/Private/Components/SkeletalMeshComponent.cpp
@@ -4905,7 +4905,7 @@ bool USkeletalMeshComponent::HandleExistingParallelEvaluationTask(bool bBlockOnT
 			{
 				(void)AutoRTFM::Close([&]
 				{
-					AutoRTFM::OnComplete([strong_this = TObjectPtr<USkeletalMeshComponent>(this), bPerformPostAnimEvaluation]
+					AutoRTFM::OnRetry([strong_this = TObjectPtr<USkeletalMeshComponent>(this), bPerformPostAnimEvaluation]
 					{
 						UE_LOG(LogSkeletalMesh, Display, TEXT("A call to `HandleExistingParallelEvaluationTask` that is blocking is causing Verse to abort and retry."));
 						const bool bBlockOnTask = true; // Was true from the if block above!
