[AutoRTFM] Add a histogram of transactional memory usage to the AutoRTFM metrics.

This will let us collect data on the game's additional memory usage when AutoRTFM is on.

#Synced-CL 50048130

diff --git a/Engine/Source/Programs/AutoRTFMTests/Private/MetricsTests.cpp b/Engine/Source/Programs/AutoRTFMTests/Private/MetricsTests.cpp
--- a/Engine/Source/Programs/AutoRTFMTests/Private/MetricsTests.cpp
+++ b/Engine/Source/Programs/AutoRTFMTests/Private/MetricsTests.cpp
@@ -5,9 +5,13 @@
 #include "AutoRTFM/Metrics.h"
 #include "AutoRTFM/Testing.h"
 #include "Catch2Includes.h"
+#include "Utils.h"
+
+#include <algorithm>
 
 TEST_CASE("Metrics.TransactionCounts")
 {
+	// Automatic retries inflate our metrics, because everything is run twice. We can get accurate histograms by disabling retry.
 	AUTORTFM_SCOPED_DISABLE_RETRY();
 
 	AutoRTFM::ResetAutoRTFMMetrics();
@@ -79,3 +83,126 @@ TEST_CASE("Metrics.TransactionCounts")
 		REQUIRE(Metrics.NumTransactionsAbortedByLanguage == 0);
 	}
 }
+
+TEST_CASE("Metrics.Histogram")
+{
+	TArray<uint8> Scratchpad;
+	Scratchpad.SetNum(10'000'000);
+
+	// Verify that our histogram bucket array is large enough for any amount of memory we could potentially allocate.
+	using HistogramBucketArray = AutoRTFM::FAutoRTFMMetrics::HistogramBucketArray;
+	static_assert(std::size(HistogramBucketArray{}) > AutoRTFM::Log2(UINTPTR_MAX));
+
+	// Automatic retries inflate our metrics, because everything is run twice. We can get accurate histograms by disabling retry.
+	AUTORTFM_SCOPED_DISABLE_RETRY();
+
+	// Start with a clean slate by resetting the metrics.
+	AutoRTFM::ResetAutoRTFMMetrics();
+	AutoRTFM::FAutoRTFMMetrics ZeroMetrics = AutoRTFM::GetAutoRTFMMetrics();
+
+	// Make sure that after a reset, the peak memory usage is zero and the histogram has no entries.
+	REQUIRE(ZeroMetrics.OverallPeakMemoryUsage == 0);
+	HistogramBucketArray ZeroHistogram = {};
+	static_assert(sizeof(ZeroMetrics.MemoryUsageHistogram) == sizeof(ZeroHistogram));
+	REQUIRE(std::equal(std::begin(ZeroMetrics.MemoryUsageHistogram), std::end(ZeroMetrics.MemoryUsageHistogram), std::begin(ZeroHistogram), std::end(ZeroHistogram)));
+
+	// Run a no-op transaction to get a baseline of "peak" memory usage when nothing is going on.
+	AutoRTFM::Testing::Commit([] {});
+	AutoRTFM::FAutoRTFMMetrics BaselineMetrics = AutoRTFM::GetAutoRTFMMetrics();
+
+	// Make sure that the histogram has exactly one entry for the current peak.
+	HistogramBucketArray BaselineHistogram = {};
+	BaselineHistogram[AutoRTFM::Log2(BaselineMetrics.OverallPeakMemoryUsage)] = 1;
+	REQUIRE(std::equal(std::begin(BaselineMetrics.MemoryUsageHistogram), std::end(BaselineMetrics.MemoryUsageHistogram), std::begin(BaselineHistogram), std::end(BaselineHistogram)));
+
+	SECTION("Abort")
+	{
+		// Run a transaction which aborts after modifying a large amount of memory. This should affect our peak
+		// by at _least_ that much memory, since the write log needs to store every byte touched.
+		AutoRTFM::Testing::Abort([&]
+		{
+			memset(Scratchpad.GetData(), 0x55, Scratchpad.NumBytes());
+			AutoRTFM::AbortTransaction();
+		});
+
+		// Make sure that AutoRTFM peak memory increased by the size of the scratchpad.
+		AutoRTFM::FAutoRTFMMetrics PostAbortMetrics = AutoRTFM::GetAutoRTFMMetrics();
+		REQUIRE(PostAbortMetrics.OverallPeakMemoryUsage > BaselineMetrics.OverallPeakMemoryUsage + Scratchpad.NumBytes());
+
+		// Make sure that the histogram has two entries--one for the current transaction, and one for the baseline.
+		HistogramBucketArray PostAbortHistogram = {};
+		PostAbortHistogram[AutoRTFM::Log2(BaselineMetrics.OverallPeakMemoryUsage)] += 1;
+		PostAbortHistogram[AutoRTFM::Log2(PostAbortMetrics.OverallPeakMemoryUsage)] += 1;
+		REQUIRE(std::equal(std::begin(PostAbortMetrics.MemoryUsageHistogram), std::end(PostAbortMetrics.MemoryUsageHistogram), std::begin(PostAbortHistogram), std::end(PostAbortHistogram)));
+	}
+
+	SECTION("Commit")
+	{
+		// Run two transactions which commit after modifying a large amount of memory.
+		for (int Count=0; Count<2; ++Count)
+		{
+			AutoRTFM::Testing::Commit([&]
+			{
+				memset(Scratchpad.GetData(), 0xCC, Scratchpad.NumBytes());
+			});
+		}
+
+		// Make sure that AutoRTFM peak memory increased by the size of the scratchpad.
+		AutoRTFM::FAutoRTFMMetrics PostCommitMetrics = AutoRTFM::GetAutoRTFMMetrics();
+		REQUIRE(PostCommitMetrics.OverallPeakMemoryUsage > BaselineMetrics.OverallPeakMemoryUsage + Scratchpad.NumBytes());
+
+		// Make sure that the histogram has three entries--two for these transactions, and one for the baseline.
+		HistogramBucketArray PostCommitHistogram = {};
+		PostCommitHistogram[AutoRTFM::Log2(BaselineMetrics.OverallPeakMemoryUsage)] += 1;
+		PostCommitHistogram[AutoRTFM::Log2(PostCommitMetrics.OverallPeakMemoryUsage)] += 2;
+		REQUIRE(std::equal(std::begin(PostCommitMetrics.MemoryUsageHistogram), std::end(PostCommitMetrics.MemoryUsageHistogram), std::begin(PostCommitHistogram), std::end(PostCommitHistogram)));
+	}
+
+	SECTION("Abort, transacting from the open")
+	{
+		// Run three transactions which abort after recording an open write on the scratchpad. (Actually modifying the scratchpad is superfluous.)
+		for (int Count=0; Count<3; ++Count)
+		{
+			AutoRTFM::TransactThenOpen([&]
+			{
+				AutoRTFM::StartTransaction();
+				AutoRTFM::RecordOpenWrite(Scratchpad.GetData(), Scratchpad.NumBytes());
+				AutoRTFM::AbortTransaction();
+			});
+		}
+
+		// Make sure that AutoRTFM peak memory increased by the size of the scratchpad.
+		AutoRTFM::FAutoRTFMMetrics PostAbortMetrics = AutoRTFM::GetAutoRTFMMetrics();
+		REQUIRE(PostAbortMetrics.OverallPeakMemoryUsage > BaselineMetrics.OverallPeakMemoryUsage + Scratchpad.NumBytes());
+
+		// Make sure that the histogram has four entries--three for these transactions, and one for the baseline.
+		HistogramBucketArray PostAbortHistogram = {};
+		PostAbortHistogram[AutoRTFM::Log2(BaselineMetrics.OverallPeakMemoryUsage)] += 1;
+		PostAbortHistogram[AutoRTFM::Log2(PostAbortMetrics.OverallPeakMemoryUsage)] += 3;
+		REQUIRE(std::equal(std::begin(PostAbortMetrics.MemoryUsageHistogram), std::end(PostAbortMetrics.MemoryUsageHistogram), std::begin(PostAbortHistogram), std::end(PostAbortHistogram)));
+	}
+
+	SECTION("Commit, transacting from the open")
+	{
+		// Run four transactions which commit after recording an open write on the scratchpad. (Actually modifying the scratchpad is superfluous.)
+		for (int Count=0; Count<4; ++Count)
+		{
+			AutoRTFM::TransactThenOpen([&]
+			{
+				AutoRTFM::StartTransaction();
+				AutoRTFM::RecordOpenWrite(Scratchpad.GetData(), Scratchpad.NumBytes());
+				AutoRTFM::CommitTransaction();
+			});
+		}
+
+		// Make sure that AutoRTFM peak memory increased by the size of the scratchpad.
+		AutoRTFM::FAutoRTFMMetrics PostCommitMetrics = AutoRTFM::GetAutoRTFMMetrics();
+		REQUIRE(PostCommitMetrics.OverallPeakMemoryUsage > BaselineMetrics.OverallPeakMemoryUsage + Scratchpad.NumBytes());
+
+		// Make sure that the histogram has five entries--four for these transactions, and one for the baseline.
+		HistogramBucketArray PostCommitHistogram = {};
+		PostCommitHistogram[AutoRTFM::Log2(BaselineMetrics.OverallPeakMemoryUsage)] += 1;
+		PostCommitHistogram[AutoRTFM::Log2(PostCommitMetrics.OverallPeakMemoryUsage)] += 4;
+		REQUIRE(std::equal(std::begin(PostCommitMetrics.MemoryUsageHistogram), std::end(PostCommitMetrics.MemoryUsageHistogram), std::begin(PostCommitHistogram), std::end(PostCommitHistogram)));
+	}
+}
diff --git a/Engine/Source/Runtime/AutoRTFM/Private/Metrics.cpp b/Engine/Source/Runtime/AutoRTFM/Private/Metrics.cpp
--- a/Engine/Source/Runtime/AutoRTFM/Private/Metrics.cpp
+++ b/Engine/Source/Runtime/AutoRTFM/Private/Metrics.cpp
@@ -22,6 +22,22 @@ FAutoRTFMMetrics GetAutoRTFMMetrics()
 	return GAutoRTFMMetrics;
 }
 
+void UpdatePeakMemoryUsageMetrics(size_t MemoryUsed)
+{
+	GAutoRTFMMetrics.OverallPeakMemoryUsage = std::max(GAutoRTFMMetrics.OverallPeakMemoryUsage, MemoryUsed);
+	GAutoRTFMMetrics.ActiveTransactionPeakMemoryUsage = std::max(GAutoRTFMMetrics.ActiveTransactionPeakMemoryUsage, MemoryUsed);
+}
+
+void UpdateMemoryUsageHistogram()
+{
+	// Histogram buckets should saturate, not overflow.
+	size_t BucketIndex = AutoRTFM::Log2(GAutoRTFMMetrics.ActiveTransactionPeakMemoryUsage);
+	uint32_t& Bucket = GAutoRTFMMetrics.MemoryUsageHistogram[BucketIndex];
+	Bucket += AUTORTFM_LIKELY(Bucket < UINT32_MAX);
+
+	GAutoRTFMMetrics.ActiveTransactionPeakMemoryUsage = 0;
+}
+
 #else  // __AUTORTFM off
 
 void ResetAutoRTFMMetrics() {}
diff --git a/Engine/Source/Runtime/AutoRTFM/Private/MetricsPriv.h b/Engine/Source/Runtime/AutoRTFM/Private/MetricsPriv.h
--- a/Engine/Source/Runtime/AutoRTFM/Private/MetricsPriv.h
+++ b/Engine/Source/Runtime/AutoRTFM/Private/MetricsPriv.h
@@ -3,6 +3,9 @@
 #pragma once
 
 #include "AutoRTFM/Metrics.h"
+#include "MemoryStats.h"
+
+#if (defined(__AUTORTFM) && __AUTORTFM)
 
 namespace AutoRTFM
 {
@@ -49,4 +52,12 @@ inline void IncrementNumTransactionsAbortedByLanguageMetric()
 	++GAutoRTFMMetrics.NumTransactionsAbortedByLanguage;
 }
 
+/** Updates the peak memory usage metrics with the passed-in value, if it's higher than the current peak value. */
+void UpdatePeakMemoryUsageMetrics(size_t MemoryUsed = FMemoryStats::GetTotalBytesAllocated());
+
+/** Updates the memory histogram metric, and resets the current-transaction peak memory usage metric. */
+void UpdateMemoryUsageHistogram();
+
 }  // namespace AutoRTFM
+
+#endif
\ No newline at end of file
diff --git a/Engine/Source/Runtime/AutoRTFM/Private/Transaction.cpp b/Engine/Source/Runtime/AutoRTFM/Private/Transaction.cpp
--- a/Engine/Source/Runtime/AutoRTFM/Private/Transaction.cpp
+++ b/Engine/Source/Runtime/AutoRTFM/Private/Transaction.cpp
@@ -9,6 +9,7 @@
 #include "CallNestInlines.h"
 #include "ContextInlines.h"
 #include "HitSet.h"
+#include "MetricsPriv.h"
 #include "Utils.h"
 
 #if PLATFORM_WINDOWS
@@ -109,6 +110,12 @@ void FTransaction::Abort()
 	Stats.Collect<EStatsKind::Abort>();
 	CollectStats();
 
+	UpdatePeakMemoryUsageMetrics();
+	if (!IsNested())
+	{
+		UpdateMemoryUsageHistogram();
+	}
+
 	// Call the destructors of all the OnCommit functors before undoing the transactional memory and
 	// calling the OnAbort callbacks. This is important as the callback functions may have captured
 	// variables that are depending on the allocated memory.
@@ -227,6 +234,11 @@ void FTransaction::CommitNested()
 	}
 #endif
 
+	// This point is likely the high-water mark for this transactional nest's memory usage.
+	// It's after we've cloned our write log and hit set into the parent transaction,
+	// but before we've cleared our deferred-task list.
+	UpdatePeakMemoryUsageMetrics();
+
 	// For all the deferred calls to `PopOnCommitHandler` that we couldn't
 	// process (because our transaction nest didn't `PushOnCommitHandler`)
 	// we need to move these to the parent now to handle them.
@@ -273,6 +285,9 @@ void FTransaction::CommitOuterNest()
 	AUTORTFM_ASSERT(IsActive());
 	AUTORTFM_ASSERT(!Parent);
 
+	UpdatePeakMemoryUsageMetrics();
+	UpdateMemoryUsageHistogram();
+
 	PreAbortTasks.Reset();
 	AbortTasks.Reset();
 
diff --git a/Engine/Source/Runtime/AutoRTFM/Public/AutoRTFM/Metrics.h b/Engine/Source/Runtime/AutoRTFM/Public/AutoRTFM/Metrics.h
--- a/Engine/Source/Runtime/AutoRTFM/Public/AutoRTFM/Metrics.h
+++ b/Engine/Source/Runtime/AutoRTFM/Public/AutoRTFM/Metrics.h
@@ -15,6 +15,19 @@ struct FAutoRTFMMetrics
 	uint64_t NumTransactionsAborted;
 	uint64_t NumTransactionsAbortedByRequest;
 	uint64_t NumTransactionsAbortedByLanguage;
+
+	uint64_t OverallPeakMemoryUsage;            // Tracks the peak memory usage across all transactions, in bytes
+	uint64_t ActiveTransactionPeakMemoryUsage;  // Tracks the peak memory usage within this transaction, in bytes
+
+	// A log2-histogram of transaction memory usage, in bytes.
+	// We do not distinguish between committed or aborted transactions in the histogram.
+	// [0] indicates the number of finished transactions using no memory.
+	// [1] indicates the number of finished transactions using a single byte of memory.
+	// [2] indicates the number of finished transactions using between 2 to 3 bytes of memory.
+	// [3] indicates the number of finished transactions using between 4 to 7 bytes of memory.
+	// and so on.
+	using HistogramBucketArray = uint32_t[64];
+	HistogramBucketArray MemoryUsageHistogram;
 };
 
 /** Reset the internal metrics tracking back to zero. */
