[AutoRTFM] Implement metrics for cascading aborts and cascading retries.

Previously, these were unrepresented in the metrics and untested.

#Synced-CL 50048130

diff --git a/Engine/Source/Programs/AutoRTFMTests/Private/MetricsTests.cpp b/Engine/Source/Programs/AutoRTFMTests/Private/MetricsTests.cpp
--- a/Engine/Source/Programs/AutoRTFMTests/Private/MetricsTests.cpp
+++ b/Engine/Source/Programs/AutoRTFMTests/Private/MetricsTests.cpp
@@ -8,6 +8,7 @@
 #include "Utils.h"
 
 #include <algorithm>
+#include <stdio.h>
 
 TEST_CASE("Metrics.TransactionCounts")
 {
@@ -26,8 +27,10 @@ TEST_CASE("Metrics.TransactionCounts")
 		REQUIRE(Metrics.NumTransactionsStarted == 2);
 		REQUIRE(Metrics.NumTransactionsCommitted == 1);
 		REQUIRE(Metrics.NumTransactionsAborted == 1);
-		REQUIRE(Metrics.NumTransactionsAbortedByRequest == 1);
-		REQUIRE(Metrics.NumTransactionsAbortedByLanguage == 0);
+		REQUIRE(Metrics.NumRequestedAborts == 1);
+		REQUIRE(Metrics.NumRequestedCascadingAborts == 0);
+		REQUIRE(Metrics.NumRequestedCascadingRetries == 0);
+		REQUIRE(Metrics.NumTransactionalViolations == 0);
 	}
 
 	SECTION("Nesting")
@@ -47,8 +50,10 @@ TEST_CASE("Metrics.TransactionCounts")
 		REQUIRE(Metrics.NumTransactionsStarted == 5);
 		REQUIRE(Metrics.NumTransactionsCommitted == 5);
 		REQUIRE(Metrics.NumTransactionsAborted == 0);
-		REQUIRE(Metrics.NumTransactionsAbortedByRequest == 0);
-		REQUIRE(Metrics.NumTransactionsAbortedByLanguage == 0);
+		REQUIRE(Metrics.NumRequestedAborts == 0);
+		REQUIRE(Metrics.NumRequestedCascadingAborts == 0);
+		REQUIRE(Metrics.NumRequestedCascadingRetries == 0);
+		REQUIRE(Metrics.NumTransactionalViolations == 0);
 	}
 
 	SECTION("Scoped rollback")
@@ -63,8 +68,10 @@ TEST_CASE("Metrics.TransactionCounts")
 		REQUIRE(Metrics.NumTransactionsStarted == 2);
 		REQUIRE(Metrics.NumTransactionsCommitted == 1);
 		REQUIRE(Metrics.NumTransactionsAborted == 1);
-		REQUIRE(Metrics.NumTransactionsAbortedByRequest == 1);
-		REQUIRE(Metrics.NumTransactionsAbortedByLanguage == 0);
+		REQUIRE(Metrics.NumRequestedAborts == 1);
+		REQUIRE(Metrics.NumRequestedCascadingAborts == 0);
+		REQUIRE(Metrics.NumRequestedCascadingRetries == 0);
+		REQUIRE(Metrics.NumTransactionalViolations == 0);
 	}
 
 	SECTION("Scoped commit")
@@ -79,8 +86,102 @@ TEST_CASE("Metrics.TransactionCounts")
 		REQUIRE(Metrics.NumTransactionsStarted == 2);
 		REQUIRE(Metrics.NumTransactionsCommitted == 2);
 		REQUIRE(Metrics.NumTransactionsAborted == 0);
-		REQUIRE(Metrics.NumTransactionsAbortedByRequest == 0);
-		REQUIRE(Metrics.NumTransactionsAbortedByLanguage == 0);
+		REQUIRE(Metrics.NumRequestedAborts == 0);
+		REQUIRE(Metrics.NumRequestedCascadingAborts == 0);
+		REQUIRE(Metrics.NumRequestedCascadingRetries == 0);
+		REQUIRE(Metrics.NumTransactionalViolations == 0);
+	}
+
+	SECTION("Cascading abort")
+	{
+		AutoRTFM::Transact([]
+		{
+			AutoRTFM::Transact([]
+			{
+				AutoRTFM::CascadingAbortTransaction();
+			});
+		});
+		AutoRTFM::FAutoRTFMMetrics Metrics = AutoRTFM::GetAutoRTFMMetrics();
+		REQUIRE(Metrics.NumTopLevelTransactionsStarted == 1);
+		REQUIRE(Metrics.NumTransactionsStarted == 2);
+		REQUIRE(Metrics.NumTransactionsCommitted == 0);
+		REQUIRE(Metrics.NumTransactionsAborted == 2);
+		REQUIRE(Metrics.NumRequestedAborts == 0);
+		REQUIRE(Metrics.NumRequestedCascadingAborts == 1);
+		REQUIRE(Metrics.NumRequestedCascadingRetries == 0);
+		REQUIRE(Metrics.NumTransactionalViolations == 0);
+	}
+
+	SECTION("Cascading retry")
+	{
+		bool bShouldRetry = true;
+
+		AutoRTFM::Transact([&]
+		{
+			AutoRTFM::OnRetry([&]
+			{
+				bShouldRetry = false;
+			});
+
+			AutoRTFM::Transact([&]
+			{
+				if (bShouldRetry)
+				{
+					AutoRTFM::CascadingRetryTransaction();
+				}
+			});
+		});
+		AutoRTFM::FAutoRTFMMetrics Metrics = AutoRTFM::GetAutoRTFMMetrics();
+		REQUIRE(Metrics.NumTopLevelTransactionsStarted == 2);
+		REQUIRE(Metrics.NumTransactionsStarted == 4);
+		REQUIRE(Metrics.NumTransactionsCommitted == 2);
+		REQUIRE(Metrics.NumTransactionsAborted == 2);
+		REQUIRE(Metrics.NumRequestedAborts == 0);
+		REQUIRE(Metrics.NumRequestedCascadingAborts == 0);
+		REQUIRE(Metrics.NumRequestedCascadingRetries == 1);
+		REQUIRE(Metrics.NumTransactionalViolations == 0);
+	}
+
+	SECTION("Transactional violation - no closed function")
+	{
+		AutoRTFMTestUtils::FScopedInternalAbortAction Scoped1(AutoRTFM::ForTheRuntime::EAutoRTFMInternalAbortActionState::Abort);
+		AutoRTFMTestUtils::FScopedEnsureOnInternalAbort Scoped2(false);
+
+		AutoRTFM::Transact([]
+		{
+			(void)fopen("fopen() is not supported in a closed transaction", "rb");
+		});
+
+		AutoRTFM::FAutoRTFMMetrics Metrics = AutoRTFM::GetAutoRTFMMetrics();
+		REQUIRE(Metrics.NumTopLevelTransactionsStarted == 1);
+		REQUIRE(Metrics.NumTransactionsStarted == 1);
+		REQUIRE(Metrics.NumTransactionsCommitted == 0);
+		REQUIRE(Metrics.NumTransactionsAborted == 1);
+		REQUIRE(Metrics.NumRequestedAborts == 0);
+		REQUIRE(Metrics.NumRequestedCascadingAborts == 0);
+		REQUIRE(Metrics.NumRequestedCascadingRetries == 0);
+		REQUIRE(Metrics.NumTransactionalViolations == 1);
+	}
+
+	SECTION("Transactional violation - unreachable")
+	{
+		AutoRTFMTestUtils::FScopedInternalAbortAction Scoped1(AutoRTFM::ForTheRuntime::EAutoRTFMInternalAbortActionState::Abort);
+		AutoRTFMTestUtils::FScopedEnsureOnInternalAbort Scoped2(false);
+
+		AutoRTFM::Transact([]
+		{
+			AutoRTFM::UnreachableIfClosed();
+		});
+
+		AutoRTFM::FAutoRTFMMetrics Metrics = AutoRTFM::GetAutoRTFMMetrics();
+		REQUIRE(Metrics.NumTopLevelTransactionsStarted == 1);
+		REQUIRE(Metrics.NumTransactionsStarted == 1);
+		REQUIRE(Metrics.NumTransactionsCommitted == 0);
+		REQUIRE(Metrics.NumTransactionsAborted == 1);
+		REQUIRE(Metrics.NumRequestedAborts == 0);
+		REQUIRE(Metrics.NumRequestedCascadingAborts == 0);
+		REQUIRE(Metrics.NumRequestedCascadingRetries == 0);
+		REQUIRE(Metrics.NumTransactionalViolations == 1);
 	}
 }
 
diff --git a/Engine/Source/Runtime/AutoRTFM/Private/Context.cpp b/Engine/Source/Runtime/AutoRTFM/Private/Context.cpp
--- a/Engine/Source/Runtime/AutoRTFM/Private/Context.cpp
+++ b/Engine/Source/Runtime/AutoRTFM/Private/Context.cpp
@@ -118,12 +118,19 @@ void FContext::AbortTransaction(ETransactionStatus Reason)
 	switch (Reason)
 	{
 		case ETransactionStatus::AbortedByLanguage:
-			IncrementNumTransactionsAbortedByLanguageMetric();
+			IncrementNumTransactionalViolationsMetric();
 			break;
 		case ETransactionStatus::AbortedByRequest:
-			IncrementNumTransactionsAbortedByRequestMetric();
+			IncrementNumRequestedAbortsMetric();
 			break;
-		default:
+		case ETransactionStatus::AbortedByCascadingAbort:
+			IncrementNumRequestedCascadingAbortsMetric();
+			break;
+		case ETransactionStatus::AbortedByCascadingRetry:
+			IncrementNumRequestedCascadingRetriesMetric();
+			break;
+		[[unlikely]] default:
+			AUTORTFM_FATAL("invalid Reason in AbortTransaction: %d", static_cast<int>(Reason));
 			break;
 	}
 
@@ -467,6 +474,10 @@ ETransactionResult FContext::Transact(void (*UninstrumentedFunction)(void*), voi
 				Status = EContextStatus::Retrying;
 				CompletionTasks.Reset();
 				RetryTasks.RemoveEachForward([](TTask<void()>& Task) { Task(); });
+
+				// A retry is treated as a new top-level transaction.
+				IncrementNumTopLevelTransactionsStartedMetric();
+				IncrementNumTransactionsStartedMetric();
 				continue;
 
 			default:
diff --git a/Engine/Source/Runtime/AutoRTFM/Private/MetricsPriv.h b/Engine/Source/Runtime/AutoRTFM/Private/MetricsPriv.h
--- a/Engine/Source/Runtime/AutoRTFM/Private/MetricsPriv.h
+++ b/Engine/Source/Runtime/AutoRTFM/Private/MetricsPriv.h
@@ -38,18 +38,32 @@ inline void IncrementNumTransactionsAbortedMetric(size_t Amount = 1)
 	GAutoRTFMMetrics.NumTransactionsAborted += Amount;
 }
 
-/** Increments counter for number of transactions aborted via user request (AbortTransaction). */
-inline void IncrementNumTransactionsAbortedByRequestMetric(size_t Amount = 1)
+/** Increments counter for number of calls to AbortTransaction. */
+inline void IncrementNumRequestedAbortsMetric(size_t Amount = 1)
 {
 	extern AutoRTFM::FAutoRTFMMetrics GAutoRTFMMetrics;
-	GAutoRTFMMetrics.NumTransactionsAbortedByRequest += Amount;
+	GAutoRTFMMetrics.NumRequestedAborts += Amount;
 }
 
-/** Increments counter for number of transactions aborted due to violating transactionality (e.g. calling a non-closed function). */
-inline void IncrementNumTransactionsAbortedByLanguageMetric(size_t Amount = 1)
+/** Increments counter for number of CascadingAbortTransaction calls. */
+inline void IncrementNumRequestedCascadingAbortsMetric(size_t Amount = 1)
 {
 	extern AutoRTFM::FAutoRTFMMetrics GAutoRTFMMetrics;
-	GAutoRTFMMetrics.NumTransactionsAbortedByLanguage += Amount;
+	GAutoRTFMMetrics.NumRequestedCascadingAborts += Amount;
+}
+
+/** Increments counter for number of CascadingRetryTransaction calls. */
+inline void IncrementNumRequestedCascadingRetriesMetric(size_t Amount = 1)
+{
+	extern AutoRTFM::FAutoRTFMMetrics GAutoRTFMMetrics;
+	GAutoRTFMMetrics.NumRequestedCascadingRetries += Amount;
+}
+
+/** Increments counter for number of transactional aborts caused by violating transactionality (unsupported calls or API misuse). */
+inline void IncrementNumTransactionalViolationsMetric(size_t Amount = 1)
+{
+	extern AutoRTFM::FAutoRTFMMetrics GAutoRTFMMetrics;
+	GAutoRTFMMetrics.NumTransactionalViolations += Amount;
 }
 
 /** Updates the peak memory usage metrics with the passed-in value, if it's higher than the current peak value. */
diff --git a/Engine/Source/Runtime/AutoRTFM/Public/AutoRTFM/Metrics.h b/Engine/Source/Runtime/AutoRTFM/Public/AutoRTFM/Metrics.h
--- a/Engine/Source/Runtime/AutoRTFM/Public/AutoRTFM/Metrics.h
+++ b/Engine/Source/Runtime/AutoRTFM/Public/AutoRTFM/Metrics.h
@@ -9,12 +9,14 @@ namespace AutoRTFM
 
 struct FAutoRTFMMetrics
 {
-	uint64_t NumTopLevelTransactionsStarted;
-	uint64_t NumTransactionsStarted;
-	uint64_t NumTransactionsCommitted;
-	uint64_t NumTransactionsAborted;
-	uint64_t NumTransactionsAbortedByRequest;
-	uint64_t NumTransactionsAbortedByLanguage;
+	uint64_t NumTopLevelTransactionsStarted;  // Tracks the number of top-level transactions that have been started
+	uint64_t NumTransactionsStarted;          // Tracks the number of transactions (at any depth) that have been started
+	uint64_t NumTransactionsCommitted;        // Tracks the number of transactions (at any depth) that have been committed
+	uint64_t NumTransactionsAborted;          // Tracks the number of transactions (at any depth) that have been aborted
+	uint64_t NumRequestedAborts;              // Tracks the number of calls to AbortTransaction()
+	uint64_t NumRequestedCascadingAborts;     // Tracks the number of calls to CascadingAbortTransaction()
+	uint64_t NumRequestedCascadingRetries;    // Tracks the number of calls to CascadingRetryTransaction()
+	uint64_t NumTransactionalViolations;      // Tracks the number of AutoRTFM aborts caused by violating transactional rules
 
 	uint64_t OverallPeakMemoryUsage;            // Tracks the peak memory usage across all transactions, in bytes
 	uint64_t ActiveTransactionPeakMemoryUsage;  // Tracks the peak memory usage within this transaction, in bytes
