[AutoRTFM] Improve metrics collection code.

Added ConvertMetricsToCSV, which converts an FAutoRTFMMetrics struct into a string suitable for inclusion in a CSV file. The matching CSV column header text is also available in FAutoRTFMMetrics::CSVHeader.

Fixed a metrics counting issue when AutoRTFM automatic retries are enabled; previously, we would miscount the number of top-level transactions and they would get out of sync with the histogram counts. Fortunately this mode is only used in AutoRTFMTests, not in the game.

#rnx

#Synced-CL 50424845

diff --git a/Engine/Source/Programs/AutoRTFMTests/Private/MetricsTests.cpp b/Engine/Source/Programs/AutoRTFMTests/Private/MetricsTests.cpp
--- a/Engine/Source/Programs/AutoRTFMTests/Private/MetricsTests.cpp
+++ b/Engine/Source/Programs/AutoRTFMTests/Private/MetricsTests.cpp
@@ -5,6 +5,7 @@
 #include "AutoRTFM/Metrics.h"
 #include "AutoRTFM/Testing.h"
 #include "Catch2Includes.h"
+#include "MetricsPriv.h"
 #include "Utils.h"
 
 #include <algorithm>
@@ -23,7 +24,7 @@ TEST_CASE("Metrics.TransactionCounts")
 		AutoRTFM::Transact([] { AutoRTFM::AbortTransaction(); });
 
 		AutoRTFM::FAutoRTFMMetrics Metrics = AutoRTFM::GetAutoRTFMMetrics();
-		REQUIRE(Metrics.NumTopLevelTransactionsStarted == 2);
+		REQUIRE(Metrics.NumTopLevelTransactionsCompleted == 2);
 		REQUIRE(Metrics.NumTransactionsStarted == 2);
 		REQUIRE(Metrics.NumTransactionsCommitted == 1);
 		REQUIRE(Metrics.NumTransactionsAborted == 1);
@@ -46,7 +47,7 @@ TEST_CASE("Metrics.TransactionCounts")
 		});
 
 		AutoRTFM::FAutoRTFMMetrics Metrics = AutoRTFM::GetAutoRTFMMetrics();
-		REQUIRE(Metrics.NumTopLevelTransactionsStarted == 1);
+		REQUIRE(Metrics.NumTopLevelTransactionsCompleted == 1);
 		REQUIRE(Metrics.NumTransactionsStarted == 5);
 		REQUIRE(Metrics.NumTransactionsCommitted == 5);
 		REQUIRE(Metrics.NumTransactionsAborted == 0);
@@ -64,7 +65,7 @@ TEST_CASE("Metrics.TransactionCounts")
 			AutoRTFM::AbortTransaction();
 		});
 		AutoRTFM::FAutoRTFMMetrics Metrics = AutoRTFM::GetAutoRTFMMetrics();
-		REQUIRE(Metrics.NumTopLevelTransactionsStarted == 1);
+		REQUIRE(Metrics.NumTopLevelTransactionsCompleted == 1);
 		REQUIRE(Metrics.NumTransactionsStarted == 2);
 		REQUIRE(Metrics.NumTransactionsCommitted == 1);
 		REQUIRE(Metrics.NumTransactionsAborted == 1);
@@ -82,7 +83,7 @@ TEST_CASE("Metrics.TransactionCounts")
 			AutoRTFM::CommitTransaction();
 		});
 		AutoRTFM::FAutoRTFMMetrics Metrics = AutoRTFM::GetAutoRTFMMetrics();
-		REQUIRE(Metrics.NumTopLevelTransactionsStarted == 1);
+		REQUIRE(Metrics.NumTopLevelTransactionsCompleted == 1);
 		REQUIRE(Metrics.NumTransactionsStarted == 2);
 		REQUIRE(Metrics.NumTransactionsCommitted == 2);
 		REQUIRE(Metrics.NumTransactionsAborted == 0);
@@ -102,7 +103,7 @@ TEST_CASE("Metrics.TransactionCounts")
 			});
 		});
 		AutoRTFM::FAutoRTFMMetrics Metrics = AutoRTFM::GetAutoRTFMMetrics();
-		REQUIRE(Metrics.NumTopLevelTransactionsStarted == 1);
+		REQUIRE(Metrics.NumTopLevelTransactionsCompleted == 1);
 		REQUIRE(Metrics.NumTransactionsStarted == 2);
 		REQUIRE(Metrics.NumTransactionsCommitted == 0);
 		REQUIRE(Metrics.NumTransactionsAborted == 2);
@@ -132,7 +133,7 @@ TEST_CASE("Metrics.TransactionCounts")
 			});
 		});
 		AutoRTFM::FAutoRTFMMetrics Metrics = AutoRTFM::GetAutoRTFMMetrics();
-		REQUIRE(Metrics.NumTopLevelTransactionsStarted == 2);
+		REQUIRE(Metrics.NumTopLevelTransactionsCompleted == 2);
 		REQUIRE(Metrics.NumTransactionsStarted == 4);
 		REQUIRE(Metrics.NumTransactionsCommitted == 2);
 		REQUIRE(Metrics.NumTransactionsAborted == 2);
@@ -153,7 +154,7 @@ TEST_CASE("Metrics.TransactionCounts")
 		});
 
 		AutoRTFM::FAutoRTFMMetrics Metrics = AutoRTFM::GetAutoRTFMMetrics();
-		REQUIRE(Metrics.NumTopLevelTransactionsStarted == 1);
+		REQUIRE(Metrics.NumTopLevelTransactionsCompleted == 1);
 		REQUIRE(Metrics.NumTransactionsStarted == 1);
 		REQUIRE(Metrics.NumTransactionsCommitted == 0);
 		REQUIRE(Metrics.NumTransactionsAborted == 1);
@@ -174,7 +175,7 @@ TEST_CASE("Metrics.TransactionCounts")
 		});
 
 		AutoRTFM::FAutoRTFMMetrics Metrics = AutoRTFM::GetAutoRTFMMetrics();
-		REQUIRE(Metrics.NumTopLevelTransactionsStarted == 1);
+		REQUIRE(Metrics.NumTopLevelTransactionsCompleted == 1);
 		REQUIRE(Metrics.NumTransactionsStarted == 1);
 		REQUIRE(Metrics.NumTransactionsCommitted == 0);
 		REQUIRE(Metrics.NumTransactionsAborted == 1);
@@ -307,3 +308,55 @@ TEST_CASE("Metrics.Histogram")
 		REQUIRE(std::equal(std::begin(PostCommitMetrics.MemoryUsageHistogram), std::end(PostCommitMetrics.MemoryUsageHistogram), std::begin(PostCommitHistogram), std::end(PostCommitHistogram)));
 	}
 }
+
+TEST_CASE("Metrics.ConvertToCSV")
+{
+	AutoRTFM::FAutoRTFMMetrics Metrics;
+
+	Metrics.NumTopLevelTransactionsCompleted  = 1001;
+	Metrics.NumTransactionsStarted            = 1002;
+	Metrics.NumTransactionsCommitted          = 1003;
+	Metrics.NumTransactionsAborted            = 1004;
+	Metrics.NumRequestedAborts                = 1005;
+	Metrics.NumRequestedCascadingAborts       = 1006;
+	Metrics.NumRequestedCascadingRetries      = 1007;
+	Metrics.NumTransactionalViolations        = 1008;
+	Metrics.OverallPeakMemoryUsage            = 1009;
+	Metrics.ActiveTransactionPeakMemoryUsage  = 1010;
+
+	uint32 Filler = 2001;
+	for (uint32_t& Entry : Metrics.MemoryUsageHistogram)
+	{
+		Entry = Filler++;
+	}
+
+	std::string Text = AutoRTFM::ConvertMetricsToCSV(Metrics);
+	REQUIRE(Text ==
+		"1001,1002,1003,1004,1005,1006,1007,1008,1009,1010,"
+		"2001,2002,2003,2004,2005,2006,2007,2008,2009,2010,"
+		"2011,2012,2013,2014,2015,2016,2017,2018,2019,2020,"
+		"2021,2022,2023,2024,2025,2026,2027,2028,2029,2030,"
+		"2031,2032,2033,2034,2035,2036,2037,2038,2039,2040,"
+		"2041,2042,2043,2044,2045,2046,2047,2048,2049,2050,"
+		"2051,2052,2053,2054,2055,2056,2057,2058,2059,2060,"
+		"2061,2062,2063,2064");
+}
+
+TEST_CASE("Metrics.AutoRetryCountsTopLevelCorrectly")
+{
+	AutoRTFM::ResetAutoRTFMMetrics();
+	AutoRTFM::Transact([&] {});
+	AutoRTFM::FAutoRTFMMetrics Metrics = AutoRTFM::GetAutoRTFMMetrics();
+
+	// If retries are enabled, this will evaluate as `2 == 1+1` due to the automatic abort-and-retry mechanism inside Transact.
+	REQUIRE(Metrics.NumTopLevelTransactionsCompleted == Metrics.NumTransactionsCommitted + Metrics.NumTransactionsAborted);
+
+	// Confirm that the number of histogram entries matches the number of top-level transactions.
+	size_t HistogramEntries = 0;
+	for (uint32_t Entry : Metrics.MemoryUsageHistogram)
+	{
+		HistogramEntries += Entry;
+	}
+
+	REQUIRE(Metrics.NumTopLevelTransactionsCompleted == HistogramEntries);
+}
\ No newline at end of file
diff --git a/Engine/Source/Runtime/AutoRTFM/Private/Context.cpp b/Engine/Source/Runtime/AutoRTFM/Private/Context.cpp
--- a/Engine/Source/Runtime/AutoRTFM/Private/Context.cpp
+++ b/Engine/Source/Runtime/AutoRTFM/Private/Context.cpp
@@ -380,7 +380,6 @@ ETransactionResult FContext::Transact(void (*UninstrumentedFunction)(void*), voi
 		AUTORTFM_ASSERT(CurrentThreadId == FThreadID::Invalid);
 		AUTORTFM_ASSERT(Status == EContextStatus::Idle);
 		AUTORTFM_ASSERT(Stack == FStackRange{});
-		IncrementNumTopLevelTransactionsStartedMetric();
 		CurrentThreadId = FThreadID::GetCurrent();
 		Stack = GetThreadStackRange();
 	}
@@ -476,7 +475,6 @@ ETransactionResult FContext::Transact(void (*UninstrumentedFunction)(void*), voi
 				RetryTasks.RemoveEachForward([](TTask<void()>& Task) { Task(); });
 
 				// A retry is treated as a new top-level transaction.
-				IncrementNumTopLevelTransactionsStartedMetric();
 				IncrementNumTransactionsStartedMetric();
 				continue;
 
diff --git a/Engine/Source/Runtime/AutoRTFM/Private/Metrics.cpp b/Engine/Source/Runtime/AutoRTFM/Private/Metrics.cpp
--- a/Engine/Source/Runtime/AutoRTFM/Private/Metrics.cpp
+++ b/Engine/Source/Runtime/AutoRTFM/Private/Metrics.cpp
@@ -4,6 +4,8 @@
 #include "MetricsPriv.h"
 #include "Utils.h"
 #include <algorithm>  // for std::max
+#include <cinttypes>  // for PRIu64
+#include <string>
 
 namespace AutoRTFM
 {
@@ -22,6 +24,30 @@ FAutoRTFMMetrics GetAutoRTFMMetrics()
 	return GAutoRTFMMetrics;
 }
 
+std::string ConvertMetricsToCSV(const FAutoRTFMMetrics& Metrics)
+{
+	char Buffer[2048];
+	int TotalLength = snprintf(Buffer, std::size(Buffer),
+		"%" PRIu64 ",%" PRIu64 ",%" PRIu64 ",%" PRIu64 ",%" PRIu64 ",%" PRIu64 ",%" PRIu64 ",%" PRIu64 ",%" PRIu64 ",%" PRIu64,
+		Metrics.NumTopLevelTransactionsCompleted, Metrics.NumTransactionsStarted, Metrics.NumTransactionsCommitted,
+		Metrics.NumTransactionsAborted, Metrics.NumRequestedAborts, Metrics.NumRequestedCascadingAborts,
+		Metrics.NumRequestedCascadingRetries, Metrics.NumTransactionalViolations, Metrics.OverallPeakMemoryUsage,
+		Metrics.ActiveTransactionPeakMemoryUsage);
+	AUTORTFM_ASSERT(TotalLength > 0);
+	AUTORTFM_ASSERT(TotalLength < static_cast<int>(std::size(Buffer)));
+
+	for (uint32_t Entry : Metrics.MemoryUsageHistogram)
+	{
+		int Appended = snprintf(Buffer + TotalLength, std::size(Buffer) - TotalLength, ",%" PRIu32, Entry);
+		AUTORTFM_ASSERT(Appended > 0);
+
+		TotalLength += Appended;
+		AUTORTFM_ASSERT(TotalLength < static_cast<int>(std::size(Buffer)));
+	}
+
+	return std::string(Buffer);
+}
+
 void UpdatePeakMemoryUsageMetrics(size_t MemoryUsed)
 {
 	GAutoRTFMMetrics.OverallPeakMemoryUsage = std::max(GAutoRTFMMetrics.OverallPeakMemoryUsage, MemoryUsed);
@@ -35,6 +61,8 @@ void UpdateMemoryUsageHistogram()
 	uint32_t& Bucket = GAutoRTFMMetrics.MemoryUsageHistogram[BucketIndex];
 	Bucket += AUTORTFM_LIKELY(Bucket < UINT32_MAX);
 
+	// This top-level transaction has completed, so increment the counter and zero out peak memory usage.
+	GAutoRTFMMetrics.NumTopLevelTransactionsCompleted++;
 	GAutoRTFMMetrics.ActiveTransactionPeakMemoryUsage = 0;
 }
 
diff --git a/Engine/Source/Runtime/AutoRTFM/Private/MetricsPriv.h b/Engine/Source/Runtime/AutoRTFM/Private/MetricsPriv.h
--- a/Engine/Source/Runtime/AutoRTFM/Private/MetricsPriv.h
+++ b/Engine/Source/Runtime/AutoRTFM/Private/MetricsPriv.h
@@ -2,6 +2,7 @@
 
 #pragma once
 
+#include "AutoRTFM/CAPI.h"
 #include "AutoRTFM/Metrics.h"
 #include "MemoryStats.h"
 
@@ -10,12 +11,8 @@
 namespace AutoRTFM
 {
 
-/** Increments counter for number of top-level transactions started. */
-inline void IncrementNumTopLevelTransactionsStartedMetric(size_t Amount = 1)
-{
-	extern AutoRTFM::FAutoRTFMMetrics GAutoRTFMMetrics;
-	GAutoRTFMMetrics.NumTopLevelTransactionsStarted += Amount;
-}
+/** Converts an AutoRTFM metrics struct into CSV-formatted text. */
+UE_AUTORTFM_API std::string ConvertMetricsToCSV(const FAutoRTFMMetrics& Metrics);
 
 /** Increments counter for number of transactions started. */
 inline void IncrementNumTransactionsStartedMetric(size_t Amount = 1)
diff --git a/Engine/Source/Runtime/AutoRTFM/Public/AutoRTFM/Metrics.h b/Engine/Source/Runtime/AutoRTFM/Public/AutoRTFM/Metrics.h
--- a/Engine/Source/Runtime/AutoRTFM/Public/AutoRTFM/Metrics.h
+++ b/Engine/Source/Runtime/AutoRTFM/Public/AutoRTFM/Metrics.h
@@ -2,40 +2,55 @@
 
 #pragma once
 
+#include "AutoRTFM/CAPI.h"
 #include <cstdint>
+#include <string>
 
 namespace AutoRTFM
 {
 
 struct FAutoRTFMMetrics
 {
-	uint64_t NumTopLevelTransactionsStarted;  // Tracks the number of top-level transactions that have been started
-	uint64_t NumTransactionsStarted;          // Tracks the number of transactions (at any depth) that have been started
-	uint64_t NumTransactionsCommitted;        // Tracks the number of transactions (at any depth) that have been committed
-	uint64_t NumTransactionsAborted;          // Tracks the number of transactions (at any depth) that have been aborted
-	uint64_t NumRequestedAborts;              // Tracks the number of calls to AbortTransaction()
-	uint64_t NumRequestedCascadingAborts;     // Tracks the number of calls to CascadingAbortTransaction()
-	uint64_t NumRequestedCascadingRetries;    // Tracks the number of calls to CascadingRetryTransaction()
-	uint64_t NumTransactionalViolations;      // Tracks the number of AutoRTFM aborts caused by violating transactional rules
+	uint64_t NumTopLevelTransactionsCompleted;  // Tracks the number of top-level transactions that have been completed
+	uint64_t NumTransactionsStarted;            // Tracks the number of transactions (at any depth) that have been started
+	uint64_t NumTransactionsCommitted;          // Tracks the number of transactions (at any depth) that have been committed
+	uint64_t NumTransactionsAborted;            // Tracks the number of transactions (at any depth) that have been aborted
+	uint64_t NumRequestedAborts;                // Tracks the number of calls to AbortTransaction()
+	uint64_t NumRequestedCascadingAborts;       // Tracks the number of calls to CascadingAbortTransaction()
+	uint64_t NumRequestedCascadingRetries;      // Tracks the number of calls to CascadingRetryTransaction()
+	uint64_t NumTransactionalViolations;        // Tracks the number of AutoRTFM aborts caused by violating transactional rules
 
 	uint64_t OverallPeakMemoryUsage;            // Tracks the peak memory usage across all transactions, in bytes
 	uint64_t ActiveTransactionPeakMemoryUsage;  // Tracks the peak memory usage within this transaction, in bytes
 
 	// A log2-histogram of transaction memory usage, in bytes.
 	// We do not distinguish between committed or aborted transactions in the histogram.
-	// [0] indicates the number of finished transactions using no memory.
-	// [1] indicates the number of finished transactions using a single byte of memory.
-	// [2] indicates the number of finished transactions using between 2 to 3 bytes of memory.
-	// [3] indicates the number of finished transactions using between 4 to 7 bytes of memory.
+	// [0] indicates the number of finished transactions using between 0 and 1 bytes of memory.
+	// [1] indicates the number of finished transactions using between 2 and 3 bytes of memory.
+	// [2] indicates the number of finished transactions using between 4 and 7 bytes of memory.
+	// [3] indicates the number of finished transactions using between 8 and 15 bytes of memory.
+	// [4] indicates the number of finished transactions using between 16 and 31 bytes of memory.
 	// and so on.
 	using HistogramBucketArray = uint32_t[64];
 	HistogramBucketArray MemoryUsageHistogram;
+
+	static constexpr char CSVHeader[] =
+		"NumTopLevelTransactionsCompleted,NumTransactionsStarted,NumTransactionsCommitted,"
+		"NumTransactionsAborted,NumRequestedAborts,NumRequestedCascadingAborts,NumRequestedCascadingRetries,"
+		"NumTransactionalViolations,OverallPeakMemoryUsage,ActiveTransactionPeakMemoryUsage,"
+		"0b+,2b+,4b+,8b+,16b+,32b+,64b+,128b+,256b+,512b+,"            // Histogram[0..9]
+		"1KB+,2KB+,4KB+,8KB+,16KB+,32KB+,64KB+,128KB+,256KB+,512KB+,"  // Histogram[10..19]
+		"1MB+,2MB+,4MB+,8MB+,16MB+,32MB+,64MB+,128MB+,256MB+,512MB+,"  // Histogram[20..29]
+		"1GB+,2GB+,4GB+,8GB+,16GB+,32GB+,64GB+,128GB+,256GB+,512GB+,"  // Histogram[30..39]
+		"1TB+,2TB+,4TB+,8TB+,16TB+,32TB+,64TB+,128TB+,256TB+,512TB+,"  // Histogram[40..49]
+		"1PB+,2PB+,4PB+,8PB+,16PB+,32PB+,64PB+,128PB+,256PB+,512PB+,"  // Histogram[50..59]
+		"1024PB+,2048PB+,4096PB+,8192PB+";                             // Histogram[60..63]
 };
 
 /** Reset the internal metrics tracking back to zero. */
-void ResetAutoRTFMMetrics();
+UE_AUTORTFM_API void ResetAutoRTFMMetrics();
 
 /** Get a snapshot of the current internal metrics. */
-FAutoRTFMMetrics GetAutoRTFMMetrics();
+UE_AUTORTFM_API FAutoRTFMMetrics GetAutoRTFMMetrics();
 
 }
