Add tests that'll be used for the soon-to-be-landed AutoRTFMTestsWithModular target.

#Synced-CL 50451909

diff --git a/Engine/Source/Programs/AutoRTFMTests/Private/InsightsTests.cpp b/Engine/Source/Programs/AutoRTFMTests/Private/InsightsTests.cpp
--- a/Engine/Source/Programs/AutoRTFMTests/Private/InsightsTests.cpp
+++ b/Engine/Source/Programs/AutoRTFMTests/Private/InsightsTests.cpp
@@ -1,5 +1,9 @@
 // Copyright Epic Games, Inc. All Rights Reserved.
 
+// These tests rely on us reaching into core using private
+// header includes which is sometimes not so fine.
+#if !defined(AUTORTFMTESTS_NO_PEEKING_AT_PRIVATES)
+
 #include "API.h"
 #include "AutoRTFM/Testing.h"
 #include "Catch2Includes.h"
@@ -553,3 +557,5 @@ TEST_CASE("Insights")
 	}
 #endif
 }
+
+#endif // !defined(AUTORTFMTESTS_NO_PEEKING_AT_PRIVATES)
diff --git a/Engine/Source/Programs/AutoRTFMTests/Private/LongJumpTests.cpp b/Engine/Source/Programs/AutoRTFMTests/Private/LongJumpTests.cpp
--- a/Engine/Source/Programs/AutoRTFMTests/Private/LongJumpTests.cpp
+++ b/Engine/Source/Programs/AutoRTFMTests/Private/LongJumpTests.cpp
@@ -1,5 +1,9 @@
 // Copyright Epic Games, Inc. All Rights Reserved.
 
+// These tests rely on us reaching into AutoRTFM using private
+// header includes which is sometimes not so fine.
+#if !defined(AUTORTFMTESTS_NO_PEEKING_AT_PRIVATES)
+
 #include "LongJump.h"
 #include "Catch2Includes.h"
 
@@ -142,3 +146,5 @@ TEST_CASE("LongJump.Benchmarks")
 		return DoNotOptimizeAway;
 	};
 }
+
+#endif // !defined(AUTORTFMTESTS_NO_PEEKING_AT_PRIVATES)
diff --git a/Engine/Source/Programs/AutoRTFMTests/Private/StatsTests.cpp b/Engine/Source/Programs/AutoRTFMTests/Private/StatsTests.cpp
--- a/Engine/Source/Programs/AutoRTFMTests/Private/StatsTests.cpp
+++ b/Engine/Source/Programs/AutoRTFMTests/Private/StatsTests.cpp
@@ -68,130 +68,133 @@ static int64 GetAutoRtfmStatValue()
 	return -1.0f;
 }
 
-TEST_CASE("Enable AutoRTFM Stats Gathering")
+TEST_CASE("StatsTests")
 {
-	// Try and register FHUDGroupManager::NewFrame but this may not work
-	// Because we may be running with Stats already enabled (and that registration only happens
-	// when we first enable Stats).  Attempt it anyway.
-	UE::Stats::DirectStatsCommand(TEXT("stat AutoRTFM_Tests"));
-
-	// Let's fill the stats history
-	for (int Num = 0; Num < 10; ++Num)
+	SECTION("Enable AutoRTFM Stats Gathering")
 	{
-		TickStatsAndWaitForFrameToBecomeAvailable();
-	}
-}
+		// Try and register FHUDGroupManager::NewFrame but this may not work
+		// Because we may be running with Stats already enabled (and that registration only happens
+		// when we first enable Stats).  Attempt it anyway.
+		UE::Stats::DirectStatsCommand(TEXT("stat AutoRTFM_Tests"));
 
-TEST_CASE("Stats FScopedCycleCounter Outside Transaction")
-{
-	SCOPE_CYCLE_COUNTER(STAT_StandardCycleCounter);
-}
+		// Let's fill the stats history
+		for (int Num = 0; Num < 10; ++Num)
+		{
+			TickStatsAndWaitForFrameToBecomeAvailable();
+		}
+	}
 
-TEST_CASE("Stats FScopedCycleCounter Inside Transaction")
-{
-	int64 PreviousValue = GetAutoRtfmStatValue();
-	AutoRTFM::Testing::Abort([&]
+	SECTION("Stats FScopedCycleCounter Outside Transaction")
 	{
 		SCOPE_CYCLE_COUNTER(STAT_StandardCycleCounter);
-		AutoRTFM::AbortTransaction();
-	});
+	}
 
-	TickStatsAndWaitForFrameToBecomeAvailable();
+	SECTION("Stats FScopedCycleCounter Inside Transaction")
+	{
+		int64 PreviousValue = GetAutoRtfmStatValue();
+		AutoRTFM::Testing::Abort([&]
+		{
+			SCOPE_CYCLE_COUNTER(STAT_StandardCycleCounter);
+			AutoRTFM::AbortTransaction();
+		});
 
-	int64 NewValue = GetAutoRtfmStatValue();
-	CATCH_CHECK(NewValue > PreviousValue);
-	PreviousValue = NewValue;
+		TickStatsAndWaitForFrameToBecomeAvailable();
 
-	AutoRTFM::Testing::Commit([&]
-	{
-		SCOPE_CYCLE_COUNTER(STAT_StandardCycleCounter);
-	});
+		int64 NewValue = GetAutoRtfmStatValue();
+		CATCH_CHECK(NewValue > PreviousValue);
+		PreviousValue = NewValue;
 
-	TickStatsAndWaitForFrameToBecomeAvailable();
-	NewValue = GetAutoRtfmStatValue();
-	CATCH_CHECK(NewValue > PreviousValue);
-}
+		AutoRTFM::Testing::Commit([&]
+		{
+			SCOPE_CYCLE_COUNTER(STAT_StandardCycleCounter);
+		});
 
-TEST_CASE("Stats Multi-FScopedCycleCounter Re-Entrant w/ Transaction")
-{
-	auto ChainCall = [] AUTORTFM_ENABLE (TFunctionRef<void()> ChainedFunc)
-	{
-		SCOPE_CYCLE_COUNTER(STAT_StandardCycleCounter);
-		ChainedFunc();
-	};
+		TickStatsAndWaitForFrameToBecomeAvailable();
+		NewValue = GetAutoRtfmStatValue();
+		CATCH_CHECK(NewValue > PreviousValue);
+	}
 
-	int64 PreviousValue = GetAutoRtfmStatValue();
-	AutoRTFM::Testing::Abort([&]
+	SECTION("Stats Multi-FScopedCycleCounter Re-Entrant w/ Transaction")
 	{
-		ChainCall([&]
+		auto ChainCall = [] AUTORTFM_ENABLE(TFunctionRef<void()> ChainedFunc)
+		{
+			SCOPE_CYCLE_COUNTER(STAT_StandardCycleCounter);
+			ChainedFunc();
+		};
+
+		int64 PreviousValue = GetAutoRtfmStatValue();
+		AutoRTFM::Testing::Abort([&]
 		{
 			ChainCall([&]
 			{
-				AutoRTFM::AbortTransaction();
+				ChainCall([&]
+				{
+					AutoRTFM::AbortTransaction();
+				});
 			});
 		});
-	});
 
-	TickStatsAndWaitForFrameToBecomeAvailable();
-	int64 NewValue = GetAutoRtfmStatValue();
-	CATCH_CHECK(NewValue > PreviousValue);
-	PreviousValue = NewValue;
+		TickStatsAndWaitForFrameToBecomeAvailable();
+		int64 NewValue = GetAutoRtfmStatValue();
+		CATCH_CHECK(NewValue > PreviousValue);
+		PreviousValue = NewValue;
 
-	AutoRTFM::Testing::Commit([&]
-	{
-		ChainCall([&]
+		AutoRTFM::Testing::Commit([&]
 		{
 			ChainCall([&]
 			{
+				ChainCall([&]
+				{
+				});
 			});
 		});
-	});
 
-	TickStatsAndWaitForFrameToBecomeAvailable();
-	NewValue = GetAutoRtfmStatValue();
-	CATCH_CHECK(NewValue > PreviousValue);
-}
+		TickStatsAndWaitForFrameToBecomeAvailable();
+		NewValue = GetAutoRtfmStatValue();
+		CATCH_CHECK(NewValue > PreviousValue);
+	}
 
-TEST_CASE("Stats FScopedCycleCounter w/ Manual Counter")
-{
-	auto ChainCall = [] AUTORTFM_ENABLE (TFunctionRef<void()> ChainedFunc)
+	SECTION("Stats FScopedCycleCounter w/ Manual Counter")
 	{
-		SCOPE_CYCLE_COUNTER(STAT_StandardCycleCounter);
-		ChainedFunc();
-	};
+		auto ChainCall = [] AUTORTFM_ENABLE(TFunctionRef<void()> ChainedFunc)
+		{
+			SCOPE_CYCLE_COUNTER(STAT_StandardCycleCounter);
+			ChainedFunc();
+		};
 
-	int64 PreviousValue = GetAutoRtfmStatValue();
-	AutoRTFM::Testing::Abort([&]
-	{
-		ChainCall([&]
+		int64 PreviousValue = GetAutoRtfmStatValue();
+		AutoRTFM::Testing::Abort([&]
+		{
+			ChainCall([&]
 			{
 				ChainCall([&]
-					{
-						AddDummyStatsMessage();
-						AutoRTFM::AbortTransaction();
-					});
+				{
+					AddDummyStatsMessage();
+					AutoRTFM::AbortTransaction();
+				});
 			});
-	});
+		});
 
-	TickStatsAndWaitForFrameToBecomeAvailable();
-	int64 NewValue = GetAutoRtfmStatValue();
-	CATCH_CHECK(NewValue > PreviousValue);
-	PreviousValue = NewValue;
+		TickStatsAndWaitForFrameToBecomeAvailable();
+		int64 NewValue = GetAutoRtfmStatValue();
+		CATCH_CHECK(NewValue > PreviousValue);
+		PreviousValue = NewValue;
 
-	AutoRTFM::Testing::Commit([&]
-	{
-		ChainCall([&]
+		AutoRTFM::Testing::Commit([&]
+		{
+			ChainCall([&]
 			{
 				ChainCall([&]
-					{
-						AddDummyStatsMessage();
-					});
+				{
+					AddDummyStatsMessage();
+				});
 			});
-	});
+		});
 
-	TickStatsAndWaitForFrameToBecomeAvailable();
-	NewValue = GetAutoRtfmStatValue();
-	CATCH_CHECK(NewValue > PreviousValue);
+		TickStatsAndWaitForFrameToBecomeAvailable();
+		NewValue = GetAutoRtfmStatValue();
+		CATCH_CHECK(NewValue > PreviousValue);
+	}
 }
 
 #endif
diff --git a/Engine/Source/Programs/AutoRTFMTests/Private/UECorePrivateTests.cpp b/Engine/Source/Programs/AutoRTFMTests/Private/UECorePrivateTests.cpp
new file mode 100644
--- /dev/null
+++ b/Engine/Source/Programs/AutoRTFMTests/Private/UECorePrivateTests.cpp
@@ -0,0 +1,225 @@
+// Copyright Epic Games, Inc. All Rights Reserved.
+
+// These tests rely on us reaching into core using private
+// header includes which is sometimes not so fine.
+#if !defined(AUTORTFMTESTS_NO_PEEKING_AT_PRIVATES)
+
+#include "AutoRTFM.h"
+#include "AutoRTFM/Testing.h"
+#include "AutoRTFMTestUtils.h"
+#include "Catch2Includes.h"
+#include "Internationalization/TextCache.h"
+#include "Internationalization/TextHistory.h"
+
+DEFINE_LOG_CATEGORY_STATIC(LogAutoRTFM_UECoreTests, Log, All)
+
+TEST_CASE("UECore.FTextHistory")
+{
+	struct AUTORTFM_ENABLE MyTextHistory final : FTextHistory_Base
+	{
+		// Need this to always return true so we hit the fun transactional bits!
+		bool CanUpdateDisplayString() override
+		{
+			return true;
+		}
+
+		MyTextHistory(const FTextId& InTextId, FString&& InSourceString) : FTextHistory_Base(InTextId, MoveTemp(InSourceString)) {}
+	};
+
+	FTextKey Namespace("NAMESPACE");
+	FTextKey Key("KEY");
+	FTextId TextId(Namespace, Key);
+	FString String("WOWWEE");
+
+	MyTextHistory History(TextId, MoveTemp(String));
+
+	SECTION("With Abort")
+	{
+		AutoRTFM::Testing::Abort([&]()
+		{
+			History.UpdateDisplayStringIfOutOfDate();
+			AutoRTFM::AbortTransaction();
+		});
+	}
+
+	SECTION("With Commit")
+	{
+		AutoRTFM::Testing::Commit([&]()
+		{
+			History.UpdateDisplayStringIfOutOfDate();
+		});
+	}
+}
+
+TEST_CASE("UECore.FTextCache")
+{
+	// FTextCache is a singleton. Grab its reference.
+	FTextCache& Cache = FTextCache::Get();
+
+	// Use a fixed cache key for the tests below.
+	const FTextId Key{ TEXT("NAMESPACE"), TEXT("KEY") };
+
+	// As FTextCache does not supply any way to query what's held in the cache,
+	// the best we can do here is to call FindOrCache() and check the returned
+	// FText strings are as expected.
+	auto CheckCacheHealthy = [&]
+		{
+			FText LookupA = Cache.FindOrCache(TEXT("VALUE"), Key);
+			REQUIRE(LookupA.ToString() == TEXT("VALUE"));
+			FText LookupB = Cache.FindOrCache(TEXT("REPLACEMENT"), Key);
+			REQUIRE(LookupB.ToString() == TEXT("REPLACEMENT"));
+			Cache.RemoveCache(Key);
+		};
+
+	SECTION("FindOrCache() Add new")
+	{
+		SECTION("With Abort")
+		{
+			AutoRTFM::Testing::Abort([&]()
+				{
+					Cache.FindOrCache(TEXT("VALUE"), Key);
+					AutoRTFM::AbortTransaction();
+				});
+
+			CheckCacheHealthy();
+		}
+
+		SECTION("With Commit")
+		{
+			AutoRTFM::Testing::Commit([&]()
+				{
+					Cache.FindOrCache(TEXT("VALUE"), Key);
+				});
+
+			CheckCacheHealthy();
+		}
+	}
+
+	SECTION("FindOrCache() Replace with same value")
+	{
+		SECTION("With Abort")
+		{
+			// Add an entry to the cache before the transaction
+			Cache.FindOrCache(TEXT("VALUE"), Key);
+
+			AutoRTFM::Testing::Abort([&]()
+				{
+					Cache.FindOrCache(TEXT("REPLACEMENT"), Key);
+					AutoRTFM::AbortTransaction();
+				});
+
+			CheckCacheHealthy();
+		}
+
+		SECTION("With Commit")
+		{
+			// Add an entry to the cache before the transaction
+			Cache.FindOrCache(TEXT("VALUE"), Key);
+
+			AutoRTFM::Testing::Commit([&]()
+				{
+					Cache.FindOrCache(TEXT("VALUE"), Key);
+				});
+
+			CheckCacheHealthy();
+		}
+	}
+
+	SECTION("FindOrCache() Replace with different value")
+	{
+		SECTION("With Abort")
+		{
+			// Add an entry to the cache before the transaction
+			Cache.FindOrCache(TEXT("ORIGINAL"), Key);
+
+			AutoRTFM::Testing::Abort([&]()
+				{
+					Cache.FindOrCache(TEXT("REPLACEMENT"), Key);
+					AutoRTFM::AbortTransaction();
+				});
+
+			CheckCacheHealthy();
+		}
+
+		SECTION("With Commit")
+		{
+			// Add an entry to the cache before the transaction
+			Cache.FindOrCache(TEXT("ORIGINAL"), Key);
+
+			AutoRTFM::Testing::Commit([&]()
+				{
+					Cache.FindOrCache(TEXT("REPLACEMENT"), Key);
+				});
+
+			CheckCacheHealthy();
+		}
+	}
+
+	static constexpr bool bSupportsTransactionalRemoveCache = false; // #jira SOL-6743
+	if (!bSupportsTransactionalRemoveCache)
+	{
+		return;
+	}
+
+	SECTION("RemoveCache()")
+	{
+		SECTION("With Abort")
+		{
+			// Add an entry to the cache before the transaction
+			Cache.FindOrCache(TEXT("VALUE"), Key);
+
+			AutoRTFM::Testing::Abort([&]()
+				{
+					Cache.RemoveCache(Key);
+					AutoRTFM::AbortTransaction();
+				});
+
+			CheckCacheHealthy();
+		}
+
+		SECTION("With Commit")
+		{
+			// Add an entry to the cache before the transaction
+			Cache.FindOrCache(TEXT("VALUE"), Key);
+
+			AutoRTFM::Testing::Commit([&]()
+				{
+					Cache.RemoveCache(Key);
+				});
+
+			CheckCacheHealthy();
+		}
+	}
+
+
+	SECTION("Mixed Closed & Open")
+	{
+		SECTION("Closed: FindOrCache() Open: RemoveCache()")
+		{
+			SECTION("With Abort")
+			{
+				AutoRTFM::Testing::Abort([&]()
+					{
+						Cache.FindOrCache(TEXT("VALUE"), Key);
+						AutoRTFM::Open([&] { Cache.RemoveCache(Key); });
+						AutoRTFM::AbortTransaction();
+					});
+
+				CheckCacheHealthy();
+			}
+
+			SECTION("With Commit")
+			{
+				AutoRTFM::Testing::Commit([&]()
+					{
+						Cache.FindOrCache(TEXT("VALUE"), Key);
+						AutoRTFM::Open([&] { Cache.RemoveCache(Key); });
+					});
+
+				CheckCacheHealthy();
+			}
+		}
+	}
+}
+
+#endif // !defined(AUTORTFMTESTS_NO_PEEKING_AT_PRIVATES)
diff --git a/Engine/Source/Programs/AutoRTFMTests/Private/UECoreTests.cpp b/Engine/Source/Programs/AutoRTFMTests/Private/UECoreTests.cpp
--- a/Engine/Source/Programs/AutoRTFMTests/Private/UECoreTests.cpp
+++ b/Engine/Source/Programs/AutoRTFMTests/Private/UECoreTests.cpp
@@ -1,12 +1,12 @@
 // Copyright Epic Games, Inc. All Rights Reserved.
 
 #include "AutoRTFM.h"
+#include "AutoRTFM/Testing.h"
 #include "AutoRTFMTestUtils.h"
 #include "Catch2Includes.h"
 #include "Async/TransactionallySafeMutex.h"
 #include "Async/ParallelFor.h"
 #include "AssetRegistry/AssetDataTagMap.h"
-#include "AutoRTFM/Testing.h"
 #include "Blueprint/BlueprintExceptionInfo.h"
 #include "Delegates/IDelegateInstance.h"
 #include "Experimental/ConcurrentLinearAllocator.h"
@@ -15,9 +15,7 @@
 #include "HAL/MallocLeakDetection.h"
 #include "HAL/ThreadHeartBeat.h"
 #include "HAL/ThreadSingleton.h"
-#include "Internationalization/TextCache.h"
 #include "Internationalization/TextFormatter.h"
-#include "Internationalization/TextHistory.h"
 #include "Internationalization/TextTransformer.h"
 #include "Logging/LogMacros.h"
 #include "Logging/StructuredLog.h"
@@ -147,44 +145,6 @@ TEST_CASE("UECore.TThreadSingleton")
 	}
 }
 
-TEST_CASE("UECore.FTextHistory")
-{
-	struct AUTORTFM_ENABLE MyTextHistory final : FTextHistory_Base
-	{
-		// Need this to always return true so we hit the fun transactional bits!
-		bool CanUpdateDisplayString() override
-		{
-			return true;
-		}
-
-		MyTextHistory(const FTextId& InTextId, FString&& InSourceString) : FTextHistory_Base(InTextId, MoveTemp(InSourceString)) {}
-	};
-
-	FTextKey Namespace("NAMESPACE");
-	FTextKey Key("KEY");
-	FTextId TextId(Namespace, Key);
-	FString String("WOWWEE");
-
-	MyTextHistory History(TextId, MoveTemp(String));
-
-	SECTION("With Abort")
-	{
-		AutoRTFM::Testing::Abort([&]()
-		{
-			History.UpdateDisplayStringIfOutOfDate();
-			AutoRTFM::AbortTransaction();
-		});
-	}
-
-	SECTION("With Commit")
-	{
-		AutoRTFM::Testing::Commit([&]()
-		{
-			History.UpdateDisplayStringIfOutOfDate();
-		});
-	}
-}
-
 TEST_CASE("UECore.FCustomVersionContainer")
 {
 	FCustomVersionContainer Container;
@@ -616,177 +576,6 @@ TEST_CASE("UECore.FText")
 	}
 }
 
-TEST_CASE("UECore.FTextCache")
-{
-	// FTextCache is a singleton. Grab its reference.
-	FTextCache& Cache = FTextCache::Get();
-
-	// Use a fixed cache key for the tests below.
-	const FTextId Key{TEXT("NAMESPACE"), TEXT("KEY")};
-
-	// As FTextCache does not supply any way to query what's held in the cache,
-	// the best we can do here is to call FindOrCache() and check the returned
-	// FText strings are as expected.
-	auto CheckCacheHealthy = [&]
-	{
-		FText LookupA = Cache.FindOrCache(TEXT("VALUE"), Key);
-		REQUIRE(LookupA.ToString() == TEXT("VALUE"));
-		FText LookupB = Cache.FindOrCache(TEXT("REPLACEMENT"), Key);
-		REQUIRE(LookupB.ToString() == TEXT("REPLACEMENT"));
-		Cache.RemoveCache(Key);
-	};
-
-	SECTION("FindOrCache() Add new")
-	{
-		SECTION("With Abort")
-		{
-			AutoRTFM::Testing::Abort([&]()
-			{
-				Cache.FindOrCache(TEXT("VALUE"), Key);
-				AutoRTFM::AbortTransaction();
-			});
-
-			CheckCacheHealthy();
-		}
-
-		SECTION("With Commit")
-		{
-			AutoRTFM::Testing::Commit([&]()
-			{
-				Cache.FindOrCache(TEXT("VALUE"), Key);
-			});
-
-			CheckCacheHealthy();
-		}
-	}
-
-	SECTION("FindOrCache() Replace with same value")
-	{
-		SECTION("With Abort")
-		{
-			// Add an entry to the cache before the transaction
-			Cache.FindOrCache(TEXT("VALUE"), Key);
-
-			AutoRTFM::Testing::Abort([&]()
-			{
-				Cache.FindOrCache(TEXT("REPLACEMENT"), Key);
-				AutoRTFM::AbortTransaction();
-			});
-
-			CheckCacheHealthy();
-		}
-
-		SECTION("With Commit")
-		{
-			// Add an entry to the cache before the transaction
-			Cache.FindOrCache(TEXT("VALUE"), Key);
-
-			AutoRTFM::Testing::Commit([&]()
-			{
-				Cache.FindOrCache(TEXT("VALUE"), Key);
-			});
-
-			CheckCacheHealthy();
-		}
-	}
-
-	SECTION("FindOrCache() Replace with different value")
-	{
-		SECTION("With Abort")
-		{
-			// Add an entry to the cache before the transaction
-			Cache.FindOrCache(TEXT("ORIGINAL"), Key);
-
-			AutoRTFM::Testing::Abort([&]()
-			{
-				Cache.FindOrCache(TEXT("REPLACEMENT"), Key);
-				AutoRTFM::AbortTransaction();
-			});
-
-			CheckCacheHealthy();
-		}
-
-		SECTION("With Commit")
-		{
-			// Add an entry to the cache before the transaction
-			Cache.FindOrCache(TEXT("ORIGINAL"), Key);
-
-			AutoRTFM::Testing::Commit([&]()
-			{
-				Cache.FindOrCache(TEXT("REPLACEMENT"), Key);
-			});
-
-			CheckCacheHealthy();
-		}
-	}
-
-	static constexpr bool bSupportsTransactionalRemoveCache = false; // #jira SOL-6743
-	if (!bSupportsTransactionalRemoveCache)
-	{
-		return;
-	}
-
-	SECTION("RemoveCache()")
-	{
-		SECTION("With Abort")
-		{
-			// Add an entry to the cache before the transaction
-			Cache.FindOrCache(TEXT("VALUE"), Key);
-
-			AutoRTFM::Testing::Abort([&]()
-			{
-				Cache.RemoveCache(Key);
-				AutoRTFM::AbortTransaction();
-			});
-
-			CheckCacheHealthy();
-		}
-
-		SECTION("With Commit")
-		{
-			// Add an entry to the cache before the transaction
-			Cache.FindOrCache(TEXT("VALUE"), Key);
-
-			AutoRTFM::Testing::Commit([&]()
-			{
-				Cache.RemoveCache(Key);
-			});
-
-			CheckCacheHealthy();
-		}
-	}
-
-
-	SECTION("Mixed Closed & Open")
-	{
-		SECTION("Closed: FindOrCache() Open: RemoveCache()")
-		{
-			SECTION("With Abort")
-			{
-				AutoRTFM::Testing::Abort([&]()
-				{
-					Cache.FindOrCache(TEXT("VALUE"), Key);
-					AutoRTFM::Open([&]{ Cache.RemoveCache(Key); });
-					AutoRTFM::AbortTransaction();
-				});
-
-				CheckCacheHealthy();
-			}
-
-			SECTION("With Commit")
-			{
-				AutoRTFM::Testing::Commit([&]()
-				{
-					Cache.FindOrCache(TEXT("VALUE"), Key);
-					AutoRTFM::Open([&]{ Cache.RemoveCache(Key); });
-				});
-
-				CheckCacheHealthy();
-			}
-		}
-	}
-}
-
 #if ENABLE_STATNAMEDEVENTS_UOBJECT
 TEST_CASE("UECore.FUObjectItem")
 {
