Add `AutoRTFMTestsWithModular` and use it to fix modular linux builds.

#Synced-CL 50588184

diff --git a/Engine/Source/Programs/AutoRTFMTests/AutoRTFMTestsWithModular.Target.cs b/Engine/Source/Programs/AutoRTFMTests/AutoRTFMTestsWithModular.Target.cs
new file mode 100644
--- /dev/null
+++ b/Engine/Source/Programs/AutoRTFMTests/AutoRTFMTestsWithModular.Target.cs
@@ -0,0 +1,16 @@
+// Copyright Epic Games, Inc. All Rights Reserved.
+
+using EpicGames.Core;
+using UnrealBuildTool;
+using System.Collections.Generic;
+using UnrealBuildBase;
+
+[SupportedPlatforms(UnrealPlatformClass.All)]
+public class AutoRTFMTestsWithModularTarget : AutoRTFMTestsTarget
+{
+	public AutoRTFMTestsWithModularTarget(TargetInfo Target) : base(Target)
+	{
+		LinkType = TargetLinkType.Modular;
+		GlobalDefinitions.Add("AUTORTFMTESTS_NO_PEEKING_AT_PRIVATES=1");
+	}
+}
diff --git a/Engine/Source/Runtime/AutoRTFM/Public/AutoRTFM/Metrics.h b/Engine/Source/Runtime/AutoRTFM/Public/AutoRTFM/Metrics.h
--- a/Engine/Source/Runtime/AutoRTFM/Public/AutoRTFM/Metrics.h
+++ b/Engine/Source/Runtime/AutoRTFM/Public/AutoRTFM/Metrics.h
@@ -2,6 +2,8 @@
 
 #pragma once
 
+#include "AutoRTFM/CAPI.h"
+
 #include <cstdint>
 
 namespace AutoRTFM
@@ -33,9 +35,9 @@ struct FAutoRTFMMetrics
 };
 
 /** Reset the internal metrics tracking back to zero. */
-void ResetAutoRTFMMetrics();
+UE_AUTORTFM_API void ResetAutoRTFMMetrics();
 
 /** Get a snapshot of the current internal metrics. */
-FAutoRTFMMetrics GetAutoRTFMMetrics();
+UE_AUTORTFM_API FAutoRTFMMetrics GetAutoRTFMMetrics();
 
 }
