Fix sanitizer issues in AutoRTFMTests.

#Synced-CL 50647680

diff --git a/Engine/Source/Programs/AutoRTFMTests/Private/Abort.cpp b/Engine/Source/Programs/AutoRTFMTests/Private/Abort.cpp
--- a/Engine/Source/Programs/AutoRTFMTests/Private/Abort.cpp
+++ b/Engine/Source/Programs/AutoRTFMTests/Private/Abort.cpp
@@ -408,7 +408,7 @@ TEST_CASE("Abort.CascadingAbortRollbackTransaction")
 	{
 		bTouched = true;
 
-		AutoRTFM::Open([&]
+		AutoRTFM::Open([&] AUTORTFM_DISABLE
 		{
 			AutoRTFM::CascadingAbortTransaction();
 			bExecutedCodeAfterCascadingRollback = true;
@@ -435,7 +435,7 @@ TEST_CASE("Abort.CascadeThroughOpen")
 			bOnAbortCallbackWasExecuted = true;
 		});
 		
-		AutoRTFM::Open([&]
+		AutoRTFM::Open([&] AUTORTFM_DISABLE
 		{
 			const AutoRTFM::ETransactionStatus Status = AutoRTFM::Close([&]
 			{
@@ -465,7 +465,7 @@ TEST_CASE("Abort.CascadeWithOnCompleteThroughOpen")
 	{
 		bTouched = true;
 
-		AutoRTFM::Open([&]
+		AutoRTFM::Open([&] AUTORTFM_DISABLE
 		{
 			const AutoRTFM::ETransactionStatus Status = AutoRTFM::Close([&]
 			{
@@ -589,7 +589,7 @@ TEST_CASE("Abort.CascadeThroughManualTransaction")
 	{
 		bTouched = true;
 
-		AutoRTFM::Open([&]
+		AutoRTFM::Open([&] AUTORTFM_DISABLE
 		{
 			AutoRTFM::StartTransaction();
 
@@ -942,7 +942,7 @@ TEST_CASE("Abort.PushOnAbortHandler_Order")
 					AutoRTFM::AbortTransaction();
 				});
 
-				AutoRTFM::Open([&]
+				AutoRTFM::Open([&] AUTORTFM_DISABLE
 				{
 					REQUIRE(Value == 101);
 				});
@@ -1135,7 +1135,7 @@ TEST_CASE("Abort.LanguageThroughOpen")
 	{
 		bTouched = true;
 
-		AutoRTFM::Open([&]
+		AutoRTFM::Open([&] AUTORTFM_DISABLE
 		{
 			const AutoRTFM::ETransactionStatus Status = AutoRTFM::Close([&]
 			{
@@ -1158,7 +1158,7 @@ TEST_CASE("Abort.StackWriteToOuterOpen")
 
 	AutoRTFM::Testing::Abort([&]
 	{
-		AutoRTFM::Open([&]
+		AutoRTFM::Open([&] AUTORTFM_DISABLE
 		{
 			std::array<int, 64> Values{};
 
@@ -1600,7 +1600,7 @@ TEST_CASE("Abort.Language.ErrorHints.MallocOpenFree")
 	{
 		void* Data = FMemory::Malloc(100);
 
-		AutoRTFM::Open([&]
+		AutoRTFM::Open([&] AUTORTFM_DISABLE
 		{
 			FMemory::Free(Data);
 		});
diff --git a/Engine/Source/Programs/AutoRTFMTests/Private/FakeVM.cpp b/Engine/Source/Programs/AutoRTFMTests/Private/FakeVM.cpp
--- a/Engine/Source/Programs/AutoRTFMTests/Private/FakeVM.cpp
+++ b/Engine/Source/Programs/AutoRTFMTests/Private/FakeVM.cpp
@@ -65,7 +65,7 @@ AutoRTFM::ETransactionResult FakeVM(const FOpcode* Opcodes, const int Length, in
 
 	return AutoRTFM::Transact([Opcodes, Length, &Value]
 	{
-		AutoRTFM::Open([Opcodes, Length, &Value]
+		AutoRTFM::Open([Opcodes, Length, &Value] AUTORTFM_DISABLE
 		{
 			std::vector<int> RecoveryOpcodeStack;
 
diff --git a/Engine/Source/Programs/AutoRTFMTests/Private/HashMapTests.cpp b/Engine/Source/Programs/AutoRTFMTests/Private/HashMapTests.cpp
--- a/Engine/Source/Programs/AutoRTFMTests/Private/HashMapTests.cpp
+++ b/Engine/Source/Programs/AutoRTFMTests/Private/HashMapTests.cpp
@@ -360,9 +360,9 @@ TEMPLATE_TEST_CASE("HashMap", "", HashMapIntToInt, HashMapIntToObj, HashMapObjTo
 	{
 		MapType Map;
 		std::unordered_map<KeyType, ValueType> Expected;
-		for (int I = 0; I < 10000; I++)
+		for (unsigned I = 0; I < 10000; I++)
 		{
-			int Mod100 = (I * 15485863) % 100;
+			unsigned Mod100 = (I * 15485863) % 100;
 			KeyType Key = KeyType{(I * 804889) % 1000};
 			switch (Mod100)
 			{
diff --git a/Engine/Source/Programs/AutoRTFMTests/Private/OpenAPI.cpp b/Engine/Source/Programs/AutoRTFMTests/Private/OpenAPI.cpp
--- a/Engine/Source/Programs/AutoRTFMTests/Private/OpenAPI.cpp
+++ b/Engine/Source/Programs/AutoRTFMTests/Private/OpenAPI.cpp
@@ -23,7 +23,7 @@ TEST_CASE("OpenAPI.StartAbortAndStartAgain")
 		// Recorded ValueB as starting at 0.
 		ValueB = 20;
 
-		AutoRTFM::Open([&]()
+		AutoRTFM::Open([&] AUTORTFM_DISABLE
 		{
 			AutoRTFM::StartTransaction();
 			AutoRTFM::RecordOpenWrite(&ValueB);
@@ -49,7 +49,7 @@ TEST_CASE("OpenAPI.CommitScopedFromOpen")
 {
 	AutoRTFM::ETransactionResult TransactResult = AutoRTFM::Transact([&]()
 	{
-		AutoRTFM::Open([&]()
+		AutoRTFM::Open([&] AUTORTFM_DISABLE
 		{
 			AutoRTFM::CommitTransaction();
 		});
@@ -78,7 +78,7 @@ TEST_CASE("OpenAPI.WriteDataInTheOpen")
 	int Value = 0;
 	auto TransactResult = AutoRTFM::Transact([&]()
 	{
-		AutoRTFM::Open([&]()
+		AutoRTFM::Open([&] AUTORTFM_DISABLE
 		{
 			AutoRTFM::Assign(Value, 1);
 		});
@@ -92,7 +92,7 @@ TEST_CASE("OpenAPI.AbortTransaction.Lambda")
 {
 	auto TransactResult = AutoRTFM::Transact([&]()
 	{
-		AutoRTFM::Open([&]()
+		AutoRTFM::Open([&] AUTORTFM_DISABLE
 		{
 			AutoRTFM::AbortTransaction();
 		});
@@ -306,7 +306,7 @@ TEST_CASE("OpenAPI.OpenCloseOpenClose")
 		if (!AutoRTFM::IsClosed()) FAIL("A - NOT CLOSED AS EXPECTED!");
 
 		// -------------------------------------
-		AutoRTFM::Open([&]()
+		AutoRTFM::Open([&] AUTORTFM_DISABLE
 		{
 			// B - WE ARE OPEN 
 			REQUIRE(!AutoRTFM::IsClosed());
@@ -318,7 +318,7 @@ TEST_CASE("OpenAPI.OpenCloseOpenClose")
 				if (!AutoRTFM::IsClosed()) FAIL("C - NOT CLOSED AS EXPECTED!");
 
 				// -------------------------------------
-				AutoRTFM::Open([&]()
+				AutoRTFM::Open([&] AUTORTFM_DISABLE
 				{
 					// D - WE ARE OPEN AGAIN
 					REQUIRE(!AutoRTFM::IsClosed());
@@ -342,7 +342,7 @@ TEST_CASE("OpenAPI.OpenCloseOpenClose")
 				AutoRTFM::AbortTransaction();
 
 				// -------------------------------------
-				AutoRTFM::Open([&]()
+				AutoRTFM::Open([&] AUTORTFM_DISABLE
 				{
 					// F - WE ARE OPEN AGAIN //
 					REQUIRE(!AutoRTFM::IsClosed());
@@ -403,7 +403,7 @@ TEST_CASE("OpenAPI.Commit_TransactOpenCloseCommit")
 	{
 		if (!AutoRTFM::IsClosed()) FAIL("Not Closed");
 
-		AutoRTFM::Open([&]()
+		AutoRTFM::Open([&] AUTORTFM_DISABLE
 		{
 			AutoRTFM::StartTransaction();
 
@@ -445,7 +445,7 @@ TEST_CASE("OpenAPI.Commit_TransactOpenCloseRollback")
 	{
 		if (!AutoRTFM::IsClosed()) FAIL("Not Closed");
 
-		AutoRTFM::Open([&]()
+		AutoRTFM::Open([&] AUTORTFM_DISABLE
 		{
 			AutoRTFM::StartTransaction();
 
@@ -541,7 +541,7 @@ TEST_CASE("OpenAPI.StackWriteCommitInTheOpen1")
 	int Value = 0;
 	AutoRTFM::Transact([&]()
 	{
-		AutoRTFM::Open([&]()
+		AutoRTFM::Open([&] AUTORTFM_DISABLE
 		{
 			AutoRTFM::StartTransaction();
 			AutoRTFM::Assign(Value, 10);
@@ -557,13 +557,13 @@ TEST_CASE("OpenAPI.StackWriteCommitInTheOpen2")
 
 	AutoRTFM::Transact([&]()
 	{
-		AutoRTFM::Open([&]()
+		AutoRTFM::Open([&] AUTORTFM_DISABLE
 		{
 			AutoRTFM::StartTransaction();
 
 			AutoRTFM::ETransactionStatus Status = AutoRTFM::Close([&]
 			{
-				AutoRTFM::Open([&]()
+				AutoRTFM::Open([&] AUTORTFM_DISABLE
 				{
 					AutoRTFM::Assign(Value, 10);
 				});
@@ -584,7 +584,7 @@ TEST_CASE("OpenAPI.Open(StartTransaction, StackWrite, AbortTransaction)")
 	bool bRanPostOpen = false;
 	AutoRTFM::Testing::Commit([&]()
 	{
-		AutoRTFM::Open([&]()
+		AutoRTFM::Open([&] AUTORTFM_DISABLE
 		{
 			AutoRTFM::StartTransaction();
 			AutoRTFM::Assign(Value, 10);
@@ -603,7 +603,7 @@ TEST_CASE("OpenAPI.Open(StackWrite, AbortTransaction)")
 	bool bRanPostOpen = false;
 	AutoRTFM::Testing::Abort([&]()
 	{
-		AutoRTFM::Open([&]()
+		AutoRTFM::Open([&] AUTORTFM_DISABLE
 		{
 			AutoRTFM::Assign(Value, 10);
 			AutoRTFM::AbortTransaction();
@@ -621,7 +621,7 @@ TEST_CASE("OpenAPI.StackWriteCommitInTheOpen3_Illegal")
 	AutoRTFM::Transact([&]()
 	{
 		int Value = 0;
-		AutoRTFM::Open([&]()
+		AutoRTFM::Open([&] AUTORTFM_DISABLE
 		{
 			AutoRTFM::StartTransaction();
 			AutoRTFM::RecordOpenWrite(&Value);
@@ -639,7 +639,7 @@ TEST_CASE("OpenAPI.WriteMemory1")
 {
 	AutoRTFM::Transact([&]()
 	{
-		AutoRTFM::Open([&]()
+		AutoRTFM::Open([&] AUTORTFM_DISABLE
 		{
 			AutoRTFM::StartTransaction();
 			
@@ -667,7 +667,7 @@ TEST_CASE("OpenAPI.WriteTrivialStructure")
 
 	AutoRTFM::Transact([&]()
 	{
-		AutoRTFM::Open([&]()
+		AutoRTFM::Open([&] AUTORTFM_DISABLE
 		{
 			AutoRTFM::StartTransaction();
 
@@ -713,7 +713,7 @@ TEST_CASE("OpenAPI.WriteTrivialStructure2")
 
 	AutoRTFM::Transact([&]()
 	{
-		AutoRTFM::Open([&]()
+		AutoRTFM::Open([&] AUTORTFM_DISABLE
 		{
 			AutoRTFM::StartTransaction();
 				
@@ -770,7 +770,7 @@ TEST_CASE("OpenAPI.Footgun1")
 			// Recorded ValueB as starting at 0.
 			ValueB = 123;
 
-			AutoRTFM::Open([&]()
+			AutoRTFM::Open([&] AUTORTFM_DISABLE
 			{
 				// Unrecorded assignments in the open
 				ValueA = 10;
@@ -808,7 +808,7 @@ TEST_CASE("OpenAPI.Footgun2")
 			// Recorded ValueB as starting at 0.
 			ValueB = 20;
 
-			AutoRTFM::Open([&]()
+			AutoRTFM::Open([&] AUTORTFM_DISABLE
 			{
 				// Unrecorded assignments in the open
 				ValueB = 10;
@@ -880,7 +880,7 @@ TEST_CASE("OpenAPI.TransOpenStartCloseAbortAbort")
 
 	AutoRTFM::Transact([&]() 
 	{
-		AutoRTFM::Open([&]()
+		AutoRTFM::Open([&] AUTORTFM_DISABLE
 		{
 			AutoRTFM::StartTransaction();
 
@@ -926,7 +926,7 @@ TEST_CASE("OpenAPI.TransOpenTransCloseAbortAbort")
 
 	AutoRTFM::Transact([&]() 
 	{
-		AutoRTFM::Open([&]()
+		AutoRTFM::Open([&] AUTORTFM_DISABLE
 		{
 			AutoRTFM::Transact([&]()
 			{
diff --git a/Engine/Source/Programs/AutoRTFMTests/Private/StackTests.cpp b/Engine/Source/Programs/AutoRTFMTests/Private/StackTests.cpp
--- a/Engine/Source/Programs/AutoRTFMTests/Private/StackTests.cpp
+++ b/Engine/Source/Programs/AutoRTFMTests/Private/StackTests.cpp
@@ -398,9 +398,9 @@ TEMPLATE_TEST_CASE("Stack", "", TrivialStack, NonTrivialStack)
 	{
 		StackType Stack;
 		std::vector<ElementType> Expected;
-		for (int I = 0; I < 10000; I++)
+		for (unsigned I = 0; I < 10000; I++)
 		{
-			int Mod100 = (I * 15485863) % 100;
+			unsigned Mod100 = (I * 15485863) % 100;
 			switch (Mod100)
 			{
 				case 0:
