[AutoRTFM] Add support for TryReadLock in the transactionally-safe read-write lock.

Inside of a transaction, we treat both read locking and write locking identically; we just hold a mutex until the end of a transaction. This made the implementation of TryReadLock pretty simple, as we could share logic with TryWriteLock.

#Synced-CL 50880216

diff --git a/Engine/Source/Programs/AutoRTFMTests/Private/RWLockTests.cpp b/Engine/Source/Programs/AutoRTFMTests/Private/RWLockTests.cpp
--- a/Engine/Source/Programs/AutoRTFMTests/Private/RWLockTests.cpp
+++ b/Engine/Source/Programs/AutoRTFMTests/Private/RWLockTests.cpp
@@ -205,26 +205,51 @@ TEMPLATE_TEST_CASE("TransactionallySafeRWLock.NestedScopeLock", "[ScopeLock]", U
 	}
 }
 
-TEMPLATE_TEST_CASE("TransactionallySafeRWLock.ScopeLock.FailedTryWriteInTransaction", "[ScopeLock]", UE::TReadScopeLock<FTransactionallySafeRWLock>, UE::TWriteScopeLock<FTransactionallySafeRWLock>)
+TEMPLATE_TEST_CASE("TransactionallySafeRWLock.ScopeLock.FailedTryInTransaction", "[ScopeLock]", UE::TReadScopeLock<FTransactionallySafeRWLock>, UE::TWriteScopeLock<FTransactionallySafeRWLock>)
 {
 	FTransactionallySafeRWLock RWLock;
 	TestType Lock(RWLock);
 
-	SECTION("Abort")
+	SECTION("TryWrite")
 	{
-		AutoRTFM::Testing::Abort([&]
+		SECTION("Abort")
 		{
-			REQUIRE(!RWLock.TryWriteLock());
-			AutoRTFM::CascadingAbortTransaction();
-		});
+			AutoRTFM::Testing::Abort([&]
+			{
+				REQUIRE(!RWLock.TryWriteLock());
+				AutoRTFM::CascadingAbortTransaction();
+			});
+		}
+
+		SECTION("Commit")
+		{
+			AutoRTFM::Testing::Commit([&]
+			{
+				REQUIRE(!RWLock.TryWriteLock());
+			});
+		}
 	}
 
-	SECTION("Commit")
+	// A perfect RW lock would allow multiple TryReads to occur safely, even inside a transaction. For now, we
+	// pessimize read locks in the closed, and so the second TryRead is not expected to succeed. #jira SOL-7429
+	SECTION("TryRead")
 	{
-		AutoRTFM::Testing::Commit([&]
+		SECTION("Abort")
 		{
-			REQUIRE(!RWLock.TryWriteLock());
-		});
+			AutoRTFM::Testing::Abort([&]
+			{
+				REQUIRE(!RWLock.TryReadLock());
+				AutoRTFM::CascadingAbortTransaction();
+			});
+		}
+
+		SECTION("Commit")
+		{
+			AutoRTFM::Testing::Commit([&]
+			{
+				REQUIRE(!RWLock.TryReadLock());
+			});
+		}
 	}
 }
 
@@ -300,6 +325,43 @@ TEST_CASE("TransactionallySafeRWLock.TryWriteLock")
 	}
 }
 
+TEST_CASE("TransactionallySafeRWLock.TryReadLock")
+{
+	SECTION("Declared Outside Transaction")
+	{
+		FTransactionallySafeRWLock RWLock;
+
+		AutoRTFM::Testing::Abort([&]
+		{
+			REQUIRE(RWLock.TryReadLock());
+			AutoRTFM::AbortTransaction();
+		});
+
+		AutoRTFM::Testing::Commit([&]
+		{
+			REQUIRE(RWLock.TryReadLock());
+			RWLock.WriteUnlock();
+		});
+	}
+
+	SECTION("Declared Inside Transaction")
+	{
+		AutoRTFM::Testing::Abort([&]
+		{
+			FTransactionallySafeRWLock RWLock;
+			REQUIRE(RWLock.TryReadLock());
+			AutoRTFM::AbortTransaction();
+		});
+
+		AutoRTFM::Testing::Commit([&]
+		{
+			FTransactionallySafeRWLock RWLock;
+			REQUIRE(RWLock.TryReadLock());
+			RWLock.WriteUnlock();
+		});
+	}
+}
+
 TEST_CASE("TransactionallySafeRWLock.NestedTryWriteLock")
 {
 	SECTION("Abort(Abort(TryWriteLock))")
@@ -361,6 +423,67 @@ TEST_CASE("TransactionallySafeRWLock.NestedTryWriteLock")
 	}
 }
 
+TEST_CASE("TransactionallySafeRWLock.NestedTryReadLock")
+{
+	SECTION("Abort(Abort(TryReadLock))")
+	{
+		AutoRTFM::Testing::Abort([&]
+		{
+			FTransactionallySafeRWLock RWLock;
+
+			AutoRTFM::Testing::Abort([&]
+			{
+				REQUIRE(RWLock.TryReadLock());
+				AutoRTFM::CascadingAbortTransaction();
+			});
+		});
+	}
+
+	SECTION("Commit(Abort(TryReadLock))")
+	{
+		AutoRTFM::Testing::Commit([&]
+		{
+			FTransactionallySafeRWLock RWLock;
+
+			AutoRTFM::Testing::Abort([&]
+			{
+				REQUIRE(RWLock.TryReadLock());
+				AutoRTFM::AbortTransaction();
+			});
+		});
+	}
+
+	SECTION("Abort(Commit(TryReadLock))")
+	{
+		AutoRTFM::Testing::Abort([&]
+		{
+			FTransactionallySafeRWLock RWLock;
+
+			AutoRTFM::Testing::Commit([&]
+			{
+				REQUIRE(RWLock.TryReadLock());
+			});
+
+			AutoRTFM::AbortTransaction();
+		});
+	}
+
+	SECTION("Commit(Commit(TryReadLock))")
+	{
+		AutoRTFM::Testing::Commit([&]
+		{
+			FTransactionallySafeRWLock RWLock;
+
+			AutoRTFM::Testing::Commit([&]
+			{
+				REQUIRE(RWLock.TryReadLock());
+			});
+
+			RWLock.ReadUnlock();
+		});
+	}
+}
+
 TEMPLATE_TEST_CASE("TransactionallySafeRWLock.RWScopeLock.ScopeLockThenTryWriteLock", "[ScopeLock]", UE::TReadScopeLock<FTransactionallySafeRWLock>, UE::TWriteScopeLock<FTransactionallySafeRWLock>)
 {
 	SECTION("Lock, Abort(TryWrite)")
@@ -387,7 +510,33 @@ TEMPLATE_TEST_CASE("TransactionallySafeRWLock.RWScopeLock.ScopeLockThenTryWriteL
 	}
 }
 
-TEST_CASE("TransactionallySafeRWLock.MultipleTryWriteLocks")
+TEMPLATE_TEST_CASE("TransactionallySafeRWLock.RWScopeLock.ScopeLockThenTryReadLock", "[ScopeLock]", UE::TReadScopeLock<FTransactionallySafeRWLock>, UE::TWriteScopeLock<FTransactionallySafeRWLock>)
+{
+	SECTION("Lock, Abort(TryRead)")
+	{
+		FTransactionallySafeRWLock RWLock;
+		TestType Lock(RWLock);
+
+		AutoRTFM::Testing::Abort([&]
+		{
+			REQUIRE(!RWLock.TryReadLock());
+			AutoRTFM::CascadingAbortTransaction();
+		});
+	}
+
+	SECTION("Lock, Commit(TryRead)")
+	{
+		FTransactionallySafeRWLock RWLock;
+		TestType Lock(RWLock);
+
+		AutoRTFM::Testing::Commit([&]
+		{
+			REQUIRE(!RWLock.TryReadLock());
+		});
+	}
+}
+
+TEST_CASE("TransactionallySafeRWLock.MultipleTryLocks")
 {
 	SECTION("TryWrite, TryWrite, Unlock")
 	{
@@ -397,6 +546,23 @@ TEST_CASE("TransactionallySafeRWLock.MultipleTryWriteLocks")
 		RWLock.WriteUnlock();
 	}
 
+	SECTION("TryWrite, TryRead, Unlock")
+	{
+		FTransactionallySafeRWLock RWLock;
+		REQUIRE(RWLock.TryWriteLock());
+		REQUIRE(!RWLock.TryReadLock());
+		RWLock.WriteUnlock();
+	}
+
+	SECTION("TryRead, TryRead, Unlock")
+	{
+		FTransactionallySafeRWLock RWLock;
+		REQUIRE(RWLock.TryReadLock());
+		REQUIRE(RWLock.TryReadLock());
+		RWLock.ReadUnlock();
+		RWLock.ReadUnlock();
+	}
+
 	SECTION("TryWrite, Commit(TryWrite), Unlock")
 	{
 		FTransactionallySafeRWLock RWLock;
@@ -408,6 +574,41 @@ TEST_CASE("TransactionallySafeRWLock.MultipleTryWriteLocks")
 		RWLock.WriteUnlock();
 	}
 
+	SECTION("TryWrite, Commit(TryRead), Unlock")
+	{
+		FTransactionallySafeRWLock RWLock;
+		REQUIRE(RWLock.TryWriteLock());
+		AutoRTFM::Testing::Commit([&]
+		{
+			REQUIRE(!RWLock.TryReadLock());
+		});
+		RWLock.WriteUnlock();
+	}
+
+	SECTION("TryRead, Commit(TryWrite), Unlock")
+	{
+		FTransactionallySafeRWLock RWLock;
+		REQUIRE(RWLock.TryReadLock());
+		AutoRTFM::Testing::Commit([&]
+		{
+			REQUIRE(!RWLock.TryWriteLock());
+		});
+		RWLock.ReadUnlock();
+	}
+
+	// A perfect RW lock would allow multiple TryReads to occur safely, even inside a transaction. For now, we
+	// pessimize read locks in the closed, and so the second TryRead is not expected to succeed. #jira SOL-7429
+	SECTION("TryRead, Commit(TryRead), Unlock")
+	{
+		FTransactionallySafeRWLock RWLock;
+		REQUIRE(RWLock.TryReadLock());
+		AutoRTFM::Testing::Commit([&]
+		{
+			REQUIRE(!RWLock.TryReadLock());
+		});
+		RWLock.ReadUnlock();
+	}
+
 	SECTION("TryWrite, Abort(TryWrite), Unlock")
 	{
 		FTransactionallySafeRWLock RWLock;
@@ -470,6 +671,19 @@ TEST_CASE("TransactionallySafeRWLock.MultipleTryWriteLocks")
 		});
 	}
 
+	// A perfect RW lock would allow multiple TryReads to occur safely, even inside a transaction. For now, we
+	// pessimize read locks in the closed, and so the second TryRead is not expected to succeed. #jira SOL-7429
+	SECTION("Abort(TryRead, TryRead)")
+	{
+		FTransactionallySafeRWLock RWLock;
+		AutoRTFM::Testing::Abort([&]
+		{
+			REQUIRE(RWLock.TryReadLock());
+			REQUIRE(!RWLock.TryReadLock());
+			AutoRTFM::AbortTransaction();
+		});
+	}
+
 	SECTION("Commit(TryWrite, TryWrite, Unlock)")
 	{
 		AutoRTFM::Testing::Commit([&]
@@ -481,6 +695,19 @@ TEST_CASE("TransactionallySafeRWLock.MultipleTryWriteLocks")
 		});
 	}
 
+	// A perfect RW lock would allow multiple TryReads to occur safely, even inside a transaction. For now, we
+	// pessimize read locks in the closed, and so the second TryRead is not expected to succeed. #jira SOL-7429
+	SECTION("Commit(TryRead, TryRead, Unlock)")
+	{
+		AutoRTFM::Testing::Commit([&]
+		{
+			FTransactionallySafeRWLock RWLock;
+			REQUIRE(RWLock.TryReadLock());
+			REQUIRE(!RWLock.TryReadLock());
+			RWLock.ReadUnlock();
+		});
+	}
+
 	SECTION("Abort(TryWrite, TryWrite, Unlock)")
 	{
 		AutoRTFM::Testing::Abort([&]
@@ -492,6 +719,20 @@ TEST_CASE("TransactionallySafeRWLock.MultipleTryWriteLocks")
 			AutoRTFM::AbortTransaction();
 		});
 	}
+
+	// A perfect RW lock would allow multiple TryReads to occur safely, even inside a transaction. For now, we
+	// pessimize read locks in the closed, and so the second TryRead is not expected to succeed. #jira SOL-7429
+	SECTION("Abort(TryRead, TryRead, Unlock)")
+	{
+		AutoRTFM::Testing::Abort([&]
+		{
+			FTransactionallySafeRWLock RWLock;
+			REQUIRE(RWLock.TryReadLock());
+			REQUIRE(!RWLock.TryReadLock());
+			RWLock.ReadUnlock();
+			AutoRTFM::AbortTransaction();
+		});
+	}
 }
 
 TEST_CASE("Benchmarks.TransactionallySafeRWLock")
diff --git a/Engine/Source/Runtime/Core/Public/Async/TransactionallySafeSharedMutex.h b/Engine/Source/Runtime/Core/Public/Async/TransactionallySafeSharedMutex.h
--- a/Engine/Source/Runtime/Core/Public/Async/TransactionallySafeSharedMutex.h
+++ b/Engine/Source/Runtime/Core/Public/Async/TransactionallySafeSharedMutex.h
@@ -128,29 +128,45 @@ struct FTransactionallySafeSharedMutexDefinition
 
 	[[nodiscard]] bool TryLock()
 	{
-		if (!AutoRTFM::IsIdle())
-		{
-			bool Ret = false;
+		return AutoRTFM::IsIdle()
+			? State->Mutex.TryLock()
+			: TryLockWithinTransaction();
+	}
+
+	[[nodiscard]] bool TryLockShared()
+	{
+		return AutoRTFM::IsIdle()
+			? State->Mutex.TryLockShared()
+			: TryLockWithinTransaction();
+	}
+
+
+private:
+	UE_NONCOPYABLE(FTransactionallySafeSharedMutexDefinition)
+
+	[[nodiscard]] bool TryLockWithinTransaction()
+	{
+		checkSlow(!AutoRTFM::IsIdle());
 
+		// The transactional system which can increment TransactionalLockCount
+		// is always single-threaded, thus this is safe to access without atomicity.
+		// For TryLock we should only lock when we have a 0 count as no one owns this lock
+		bool LockTaken = false;
+
+		if (State->TransactionalLockCount == 0)
+		{
 			UE_AUTORTFM_OPEN
 			{
-				// The transactional system which can increment TransactionalLockCount
-				// is always single-threaded, thus this is safe to check without atomicity.
-				// For TryLock we should only lock when we have a 0 count as no one owns this lock
-				if (0 == State->TransactionalLockCount)
-				{
-					Ret = State->Mutex.TryLock();
-				}
-
-				if (Ret)
+				if (State->Mutex.TryLock())
 				{
-					State->TransactionalLockCount += 1;
+					LockTaken = true;
+					State->TransactionalLockCount = 1;
 				}
 			};
 
-			// Only setup the OnAbort if we *did* grab a lock; otherwise, we will not want to do 
+			// Only setup the OnAbort if we *did* grab a lock; otherwise, we will not want to do
 			// anything with the count or lock.
-			if (Ret)
+			if (LockTaken)
 			{
 				// We explicitly copy the state here for the case that `this` was stack
 				// allocated and has already died before the on-abort is hit.
@@ -159,16 +175,11 @@ struct FTransactionallySafeSharedMutexDefinition
 					State.Object->Unlock();
 				};
 			}
-
-			return Ret;
 		}
 
-		return State->Mutex.TryLock();
+		return LockTaken;
 	}
 
-private:
-	UE_NONCOPYABLE(FTransactionallySafeSharedMutexDefinition)
-
 	struct FState final
 	{
 		::UE::FSharedMutex Mutex;
diff --git a/Engine/Source/Runtime/Core/Public/Misc/TransactionallySafeRWLock.h b/Engine/Source/Runtime/Core/Public/Misc/TransactionallySafeRWLock.h
--- a/Engine/Source/Runtime/Core/Public/Misc/TransactionallySafeRWLock.h
+++ b/Engine/Source/Runtime/Core/Public/Misc/TransactionallySafeRWLock.h
@@ -41,6 +41,11 @@ public:
 		Super::Unlock();
 	}
 
+	bool TryReadLock()
+	{
+		return Super::TryLockShared();
+	}
+
 	bool TryWriteLock()
 	{
 		return Super::TryLock();
