[AutoRTFM] Fixes for FSanitizer's use of TIntervalTree

Correctly handle interval insertions:
* If the data does not match, cut and replace the existing intervals.
* If the data does match then otherwise glob them together.
This fixes the "most recent write" when reporting a memory violation.

Correctly return all intersecting memory writes when querying. Previously we could report false negatives if the write partially overlapped the queried start.

Print the first 4 writes that overlap the open write, on sanitizer failure.

#Synced-CL 51461940

diff --git a/Engine/Source/Programs/AutoRTFMTests/Private/IntervalTreeTests.cpp b/Engine/Source/Programs/AutoRTFMTests/Private/IntervalTreeTests.cpp
--- a/Engine/Source/Programs/AutoRTFMTests/Private/IntervalTreeTests.cpp
+++ b/Engine/Source/Programs/AutoRTFMTests/Private/IntervalTreeTests.cpp
@@ -24,12 +24,18 @@ TEST_CASE("IntervalTree")
 		Tree.Insert(Ptr(Address), Size, Data);
 	};
 
+	// Returns true if Overlaps() finds at least one overlapping interval for [Address, Address+Size).
+	auto HasOverlap = [Ptr](const TIntervalTree<int>& Tree, uintptr_t Address, uintptr_t Size) -> bool
+	{
+		return Tree.Overlaps(Ptr(Address), Size);
+	};
+
 	// Checks that the interval tree has all of the expected intervals, and that Get() returns
 	// the correct value for various intervals.
 	auto Validate = [&] AUTORTFM_DISABLE (TIntervalTree<int>& IntervalTree, std::vector<FIntervalAndData> Expected)
 	{
 		std::vector<FIntervalAndData> Got;
-		for (auto Entry : IntervalTree)
+		for (TIntervalTree<int>::FIntervalAndData Entry : IntervalTree)
 		{
 			Got.push_back(FIntervalAndData{Entry.Interval, Entry.Data});
 		}
@@ -47,10 +53,11 @@ TEST_CASE("IntervalTree")
 			const FIntervalAndData Entry = Expected[Index];
 			const bool bTouchingPrev = Index > 0 && Expected[Index - 1].Interval.End == Entry.Interval.Start;
 			const bool bTouchingNext = Index < (Expected.size() - 1) && Expected[Index + 1].Interval.Start == Entry.Interval.End;
-			REQUIRE(nullptr != IntervalTree.Get(Ptr(Entry.Interval.Start), 1));
-			REQUIRE(nullptr != IntervalTree.Get(Ptr(Entry.Interval.Start), Entry.Interval.Size()));
-			REQUIRE(bTouchingPrev == (nullptr != IntervalTree.Get(Ptr(Entry.Interval.Start - 1), 1)));
-			REQUIRE(bTouchingNext == (nullptr != IntervalTree.Get(Ptr(Entry.Interval.End), 1)));
+			REQUIRE(HasOverlap(IntervalTree, Entry.Interval.Start, 1));
+			REQUIRE(HasOverlap(IntervalTree, Entry.Interval.Start, Entry.Interval.Size()));
+			REQUIRE((*IntervalTree.Get(Ptr(Entry.Interval.Start), Entry.Interval.Size()).begin()).Data == Entry.Data);
+			REQUIRE(bTouchingPrev == HasOverlap(IntervalTree, Entry.Interval.Start - 1, 1));
+			REQUIRE(bTouchingNext == HasOverlap(IntervalTree, Entry.Interval.End, 1));
 		}
 	};
 
@@ -63,7 +70,7 @@ TEST_CASE("IntervalTree")
 		IntervalTree.Reset();
 		Validate(IntervalTree, {});
 
-		// New intervals
+		// New non-overlapping intervals
 		Insert(IntervalTree, 40, 10, 1);
 		Validate(IntervalTree, {{{40, 50}, 1}});
 
@@ -76,49 +83,222 @@ TEST_CASE("IntervalTree")
 		Insert(IntervalTree, 80, 10, 4);
 		Validate(IntervalTree, {{{20, 30}, 2}, {{40, 50}, 1}, {{60, 70}, 3}, {{80, 90}, 4}});
 
-		// Re-inserted intervals
-		// TODO(SOL-8824): This is incorrect for intervals with data.
-		// The existing interval should be updated with the new data.
+		// Re-inserting with different data replaces the data of the existing interval.
 		Insert(IntervalTree, 20, 10, 5);
-		Validate(IntervalTree, {{{20, 30}, 2}, {{40, 50}, 1}, {{60, 70}, 3}, {{80, 90}, 4}});
+		Validate(IntervalTree, {{{20, 30}, 5}, {{40, 50}, 1}, {{60, 70}, 3}, {{80, 90}, 4}});
 
 		Insert(IntervalTree, 40, 10, 6);
-		Validate(IntervalTree, {{{20, 30}, 2}, {{40, 50}, 1}, {{60, 70}, 3}, {{80, 90}, 4}});
+		Validate(IntervalTree, {{{20, 30}, 5}, {{40, 50}, 6}, {{60, 70}, 3}, {{80, 90}, 4}});
 
 		Insert(IntervalTree, 60, 10, 7);
-		Validate(IntervalTree, {{{20, 30}, 2}, {{40, 50}, 1}, {{60, 70}, 3}, {{80, 90}, 4}});
+		Validate(IntervalTree, {{{20, 30}, 5}, {{40, 50}, 6}, {{60, 70}, 7}, {{80, 90}, 4}});
 
 		Insert(IntervalTree, 80, 10, 8);
-		Validate(IntervalTree, {{{20, 30}, 2}, {{40, 50}, 1}, {{60, 70}, 3}, {{80, 90}, 4}});
+		Validate(IntervalTree, {{{20, 30}, 5}, {{40, 50}, 6}, {{60, 70}, 7}, {{80, 90}, 8}});
 
-		// Expanding intervals without overlap
-		// TODO(SOL-8824): This is incorrect for intervals with data.
-		// The existing interval should be reduced to no-overlaps, and the new
-		// interval should be inserted.
+		// Re-inserting with same data is effectively a no-op.
+		Insert(IntervalTree, 20, 10, 5);
+		Validate(IntervalTree, {{{20, 30}, 5}, {{40, 50}, 6}, {{60, 70}, 7}, {{80, 90}, 8}});
+
+		// Inserting adjacent intervals with different data: intervals stay separate (no merge).
 		Insert(IntervalTree, 15, 5, 9);
-		Validate(IntervalTree, {{{15, 30}, 2}, {{40, 50}, 1}, {{60, 70}, 3}, {{80, 90}, 4}});
+		Validate(IntervalTree, {{{15, 20}, 9}, {{20, 30}, 5}, {{40, 50}, 6}, {{60, 70}, 7}, {{80, 90}, 8}});
 
 		Insert(IntervalTree, 30, 5, 10);
-		Validate(IntervalTree, {{{15, 35}, 2}, {{40, 50}, 1}, {{60, 70}, 3}, {{80, 90}, 4}});
+		Validate(IntervalTree, {{{15, 20}, 9}, {{20, 30}, 5}, {{30, 35}, 10}, {{40, 50}, 6}, {{60, 70}, 7}, {{80, 90}, 8}});
 
 		Insert(IntervalTree, 55, 5, 11);
-		Validate(IntervalTree, {{{15, 35}, 2}, {{40, 50}, 1}, {{55, 70}, 3}, {{80, 90}, 4}});
+		Validate(IntervalTree, {{{15, 20}, 9}, {{20, 30}, 5}, {{30, 35}, 10}, {{40, 50}, 6}, {{55, 60}, 11}, {{60, 70}, 7}, {{80, 90}, 8}});
 
 		Insert(IntervalTree, 70, 5, 12);
-		Validate(IntervalTree, {{{15, 35}, 2}, {{40, 50}, 1}, {{55, 75}, 3}, {{80, 90}, 4}});
+		Validate(IntervalTree, {{{15, 20}, 9}, {{20, 30}, 5}, {{30, 35}, 10}, {{40, 50}, 6}, {{55, 60}, 11}, {{60, 70}, 7}, {{70, 75}, 12}, {{80, 90}, 8}});
 
-		// Expanding intervals with overlap
-		// TODO(SOL-8824): This is incorrect for intervals with data.
-		// The existing interval should be reduced to no-overlaps, and the new
-		// interval should be inserted.
+		// Inserting between adjacent different-data intervals: new interval stays separate from both neighbours.
 		Insert(IntervalTree, 35, 5, 13);
-		Validate(IntervalTree, {{{15, 35}, 2}, {{35, 50}, 1}, {{55, 75}, 3}, {{80, 90}, 4}});
+		Validate(IntervalTree, {{{15, 20}, 9}, {{20, 30}, 5}, {{30, 35}, 10}, {{35, 40}, 13}, {{40, 50}, 6}, {{55, 60}, 11}, {{60, 70}, 7}, {{70, 75}, 12}, {{80, 90}, 8}});
 
 		Insert(IntervalTree, 50, 5, 14);
-		Validate(IntervalTree, {{{15, 35}, 2}, {{35, 55}, 1}, {{55, 75}, 3}, {{80, 90}, 4}});
+		Validate(IntervalTree, {{{15, 20}, 9}, {{20, 30}, 5}, {{30, 35}, 10}, {{35, 40}, 13}, {{40, 50}, 6}, {{50, 55}, 14}, {{55, 60}, 11}, {{60, 70}, 7}, {{70, 75}, 12}, {{80, 90}, 8}});
 
 		Insert(IntervalTree, 75, 5, 15);
-		Validate(IntervalTree, {{{15, 35}, 2}, {{35, 55}, 1}, {{55, 80}, 3}, {{80, 90}, 4}});
+		Validate(IntervalTree, {{{15, 20}, 9}, {{20, 30}, 5}, {{30, 35}, 10}, {{35, 40}, 13}, {{40, 50}, 6}, {{50, 55}, 14}, {{55, 60}, 11}, {{60, 70}, 7}, {{70, 75}, 12}, {{75, 80}, 15}, {{80, 90}, 8}});
+	}
+
+	SECTION("Insert.SameData")
+	{
+		// Exact match, same data — effectively a no-op.
+		Insert(IntervalTree, 40, 10, 1);
+		Insert(IntervalTree, 40, 10, 1);
+		Validate(IntervalTree, {{{40, 50}, 1}});
+
+		IntervalTree.Reset();
+
+		// New interval is a strict subset of an existing interval with the same data.
+		// The existing interval absorbs the new one (no change).
+		Insert(IntervalTree, 30, 30, 1); // [30,60)
+		Insert(IntervalTree, 40, 10, 1); // [40,50) subset
+		Validate(IntervalTree, {{{30, 60}, 1}});
+
+		IntervalTree.Reset();
+
+		// New interval is a strict superset — absorbs the existing interval.
+		Insert(IntervalTree, 40, 10, 1); // [40,50)
+		Insert(IntervalTree, 30, 30, 1); // [30,60) superset
+		Validate(IntervalTree, {{{30, 60}, 1}});
+
+		IntervalTree.Reset();
+
+		// Left-adjacent, same data — two intervals merge into one.
+		Insert(IntervalTree, 30, 10, 1); // [30,40)
+		Insert(IntervalTree, 40, 10, 1); // [40,50) adjacent on right
+		Validate(IntervalTree, {{{30, 50}, 1}});
+
+		IntervalTree.Reset();
+
+		// Right-adjacent, same data — two intervals merge into one.
+		Insert(IntervalTree, 40, 10, 1); // [40,50)
+		Insert(IntervalTree, 30, 10, 1); // [30,40) adjacent on left
+		Validate(IntervalTree, {{{30, 50}, 1}});
+
+		IntervalTree.Reset();
+
+		// Both-adjacent, same data — three intervals merge into one.
+		Insert(IntervalTree, 20, 10, 1); // [20,30)
+		Insert(IntervalTree, 40, 10, 1); // [40,50)
+		Insert(IntervalTree, 30, 10, 1); // [30,40) fills the gap
+		Validate(IntervalTree, {{{20, 50}, 1}});
+
+		IntervalTree.Reset();
+
+		// Left partial overlap, same data — predecessor extends into the new interval; merge left.
+		Insert(IntervalTree, 30, 15, 1); // [30,45)
+		Insert(IntervalTree, 40, 10, 1); // [40,50) overlaps right side of existing
+		Validate(IntervalTree, {{{30, 50}, 1}});
+
+		IntervalTree.Reset();
+
+		// Right partial overlap, same data — new interval extends into a successor; merge right.
+		Insert(IntervalTree, 45, 15, 1); // [45,60)
+		Insert(IntervalTree, 40, 10, 1); // [40,50) overlaps left side of existing
+		Validate(IntervalTree, {{{40, 60}, 1}});
+
+		IntervalTree.Reset();
+
+		// Many non-adjacent same-data intervals + a spanning insert — all globbed into one.
+		Insert(IntervalTree, 10, 10, 1); // [10,20)
+		Insert(IntervalTree, 30, 10, 1); // [30,40)
+		Insert(IntervalTree, 50, 10, 1); // [50,60)
+		Insert(IntervalTree, 70, 10, 1); // [70,80)
+		Insert(IntervalTree, 90, 10, 1); // [90,100)
+		Validate(IntervalTree, {{{10, 20}, 1}, {{30, 40}, 1}, {{50, 60}, 1}, {{70, 80}, 1}, {{90, 100}, 1}});
+		Insert(IntervalTree, 5, 100, 1);  // [5,105) spans all five
+		Validate(IntervalTree, {{{5, 105}, 1}});
+	}
+
+	SECTION("Insert.DifferentData")
+	{
+		// Exact match, different data — data is replaced.
+		Insert(IntervalTree, 40, 10, 1);
+		Insert(IntervalTree, 40, 10, 2);
+		Validate(IntervalTree, {{{40, 50}, 2}});
+
+		IntervalTree.Reset();
+
+		// New interval is a strict subset, different data — existing is split into three.
+		Insert(IntervalTree, 30, 30, 1); // [30,60)
+		Insert(IntervalTree, 40, 10, 2); // [40,50) subset with different data
+		Validate(IntervalTree, {{{30, 40}, 1}, {{40, 50}, 2}, {{50, 60}, 1}});
+
+		IntervalTree.Reset();
+
+		// New interval is a strict superset, different data — existing is deleted.
+		Insert(IntervalTree, 40, 10, 1); // [40,50)
+		Insert(IntervalTree, 30, 30, 2); // [30,60) superset with different data
+		Validate(IntervalTree, {{{30, 60}, 2}});
+
+		IntervalTree.Reset();
+
+		// Left partial overlap, different data — existing end is trimmed.
+		Insert(IntervalTree, 30, 15, 1); // [30,45)
+		Insert(IntervalTree, 40, 10, 2); // [40,50) overlaps right side of existing
+		Validate(IntervalTree, {{{30, 40}, 1}, {{40, 50}, 2}});
+
+		IntervalTree.Reset();
+
+		// Right partial overlap, different data — existing start is trimmed (via remove+re-insert).
+		Insert(IntervalTree, 45, 15, 1); // [45,60)
+		Insert(IntervalTree, 40, 10, 2); // [40,50) overlaps left side of existing
+		Validate(IntervalTree, {{{40, 50}, 2}, {{50, 60}, 1}});
+
+		IntervalTree.Reset();
+
+		// Adjacent left, different data — intervals stay separate (no merge).
+		Insert(IntervalTree, 30, 10, 1); // [30,40)
+		Insert(IntervalTree, 40, 10, 2); // [40,50)
+		Validate(IntervalTree, {{{30, 40}, 1}, {{40, 50}, 2}});
+
+		IntervalTree.Reset();
+
+		// Adjacent right, different data — intervals stay separate (no merge).
+		Insert(IntervalTree, 50, 10, 1); // [50,60)
+		Insert(IntervalTree, 40, 10, 2); // [40,50)
+		Validate(IntervalTree, {{{40, 50}, 2}, {{50, 60}, 1}});
+	}
+
+	SECTION("Insert.GlobMany")
+	{
+		// Insert a span with the same data as all existing intervals — everything merges into one.
+		Insert(IntervalTree, 10, 10, 1); // [10,20)
+		Insert(IntervalTree, 30, 10, 1); // [30,40)
+		Insert(IntervalTree, 50, 10, 1); // [50,60)
+		Insert(IntervalTree, 70, 10, 1); // [70,80)
+		Insert(IntervalTree, 90, 10, 1); // [90,100)
+		Insert(IntervalTree, 5, 100, 1); // [5,105) with same data
+		Validate(IntervalTree, {{{5, 105}, 1}});
+
+		IntervalTree.Reset();
+
+		// Insert a span that mixes same-data and different-data neighbours.
+		// Same-data intervals are merged in; different-data intervals within the range are deleted.
+		// The effective span expands to absorb all touching same-data intervals.
+		Insert(IntervalTree, 10, 10, 1); // [10,20) — same data as the spanning insert
+		Insert(IntervalTree, 30, 10, 2); // [30,40) — different data
+		Insert(IntervalTree, 50, 10, 1); // [50,60) — same data
+		Insert(IntervalTree, 70, 10, 2); // [70,80) — different data
+		Insert(IntervalTree, 90, 10, 1); // [90,100) — same data
+		// Insert [15,95) with data=1. Its effective range expands:
+		//   - [10,20),1 predecessor overlaps from left -> absorbed -> EffStart=10
+		//   - [30,40),2 fully contained, different data -> deleted
+		//   - [50,60),1 same data -> absorbed (EffEnd stays 95)
+		//   - [70,80),2 fully contained, different data -> deleted
+		//   - [90,100),1 overlaps past EffEnd=95, same data -> absorbed -> EffEnd=100
+		// Result: one interval [10,100) with data=1.
+		Insert(IntervalTree, 15, 80, 1); // [15,95)
+		Validate(IntervalTree, {{{10, 100}, 1}});
+	}
+
+	SECTION("Insert.StressRemove")
+	{
+		// Insert many non-adjacent same-data intervals, then insert a span covering
+		// them all. This exercises Remove() with potential two-children swap during
+		// Phase B's forward scan, and verifies the LowerBound re-query approach.
+		for (int I = 0; I < 100; I++)
+		{
+			Insert(IntervalTree, 100 + static_cast<uintptr_t>(I) * 20, 10, 1);
+		}
+		Insert(IntervalTree, 100, 1990, 1); // [100, 2090) spans all 100
+		Validate(IntervalTree, {{{100, 2090}, 1}});
+
+		IntervalTree.Reset();
+
+		// Same but with different data values to test fully-contained deletion
+		// of many nodes during a single Insert call.
+		for (int I = 0; I < 100; I++)
+		{
+			Insert(IntervalTree, 100 + static_cast<uintptr_t>(I) * 20, 10, I);
+		}
+		Insert(IntervalTree, 100, 1990, 999); // [100, 2090) replaces all 100
+		Validate(IntervalTree, {{{100, 2090}, 999}});
 	}
 
 	SECTION("Merge")
@@ -145,37 +325,66 @@ TEST_CASE("IntervalTree")
 
 		SECTION("Expand")
 		{
+			// IntervalTreeB's intervals are adjacent to but not overlapping those in IntervalTree,
+			// and all have different data. After merge, all intervals remain separate — no merging occurs.
 			TIntervalTree<int> IntervalTreeB;
-			Insert(IntervalTreeB, 15, 5, 10);
-			Insert(IntervalTreeB, 30, 3, 11);
-			Insert(IntervalTreeB, 37, 3, 12);
-			Insert(IntervalTreeB, 50, 2, 13);
-			Insert(IntervalTreeB, 55, 5, 14);
-			Insert(IntervalTreeB, 70, 4, 15);
-			Insert(IntervalTreeB, 78, 2, 16);
-			Insert(IntervalTreeB, 90, 3, 17);
+			Insert(IntervalTreeB, 15, 5, 10);  // [15,20) adjacent-left to [20,30),1
+			Insert(IntervalTreeB, 30, 3, 11);  // [30,33) adjacent-right to [20,30),1
+			Insert(IntervalTreeB, 37, 3, 12);  // [37,40) adjacent-left to [40,50),2
+			Insert(IntervalTreeB, 50, 2, 13);  // [50,52) adjacent-right to [40,50),2
+			Insert(IntervalTreeB, 55, 5, 14);  // [55,60) adjacent-left to [60,70),3
+			Insert(IntervalTreeB, 70, 4, 15);  // [70,74) adjacent-right to [60,70),3
+			Insert(IntervalTreeB, 78, 2, 16);  // [78,80) adjacent-left to [80,90),4
+			Insert(IntervalTreeB, 90, 3, 17);  // [90,93) adjacent-right to [80,90),4
 			Validate(IntervalTreeB, {{{15, 20}, 10}, {{30, 33}, 11}, {{37, 40}, 12}, {{50, 52}, 13}, {{55, 60}, 14}, {{70, 74}, 15}, {{78, 80}, 16}, {{90, 93}, 17}});
 
-			// TODO(SOL-8824): This is incorrect for intervals with data.
-			// The existing interval should be reduced to no-overlaps, and the new
-			// interval should be inserted.
 			IntervalTree.Merge(IntervalTreeB);
-			Validate(IntervalTree, {{{15, 33}, 1}, {{37, 52}, 2}, {{55, 74}, 3}, {{78, 93}, 4}});
+			Validate(IntervalTree, {
+				{{15, 20}, 10}, {{20, 30}, 1}, {{30, 33}, 11}, {{37, 40}, 12}, {{40, 50}, 2},
+				{{50, 52}, 13}, {{55, 60}, 14}, {{60, 70}, 3}, {{70, 74}, 15}, {{78, 80}, 16},
+				{{80, 90}, 4}, {{90, 93}, 17}
+			});
 		}
 
 		SECTION("Overlap")
 		{
+			// IntervalTreeB's intervals exactly fill the gaps between IntervalTree's intervals,
+			// touching both neighbours at their endpoints. All data values are different, so
+			// no merging occurs and all intervals remain as separate entries.
 			TIntervalTree<int> IntervalTreeB;
-			Insert(IntervalTreeB, 30, 10, 10);
-			Insert(IntervalTreeB, 50, 10, 11);
-			Insert(IntervalTreeB, 70, 10, 12);
+			Insert(IntervalTreeB, 30, 10, 10); // [30,40) fills gap between [20,30) and [40,50)
+			Insert(IntervalTreeB, 50, 10, 11); // [50,60) fills gap between [40,50) and [60,70)
+			Insert(IntervalTreeB, 70, 10, 12); // [70,80) fills gap between [60,70) and [80,90)
 			Validate(IntervalTreeB, {{{30, 40}, 10}, {{50, 60}, 11}, {{70, 80}, 12}});
 
-			// TODO(SOL-8824): This is incorrect for intervals with data.
-			// The existing interval should be reduced to no-overlaps, and the new
-			// interval should be inserted.
 			IntervalTree.Merge(IntervalTreeB);
-			Validate(IntervalTree, {{{20, 30}, 1}, {{30, 60}, 2}, {{60, 80}, 3}, {{80, 90}, 4}});
+			Validate(IntervalTree, {
+				{{20, 30}, 1}, {{30, 40}, 10}, {{40, 50}, 2},
+				{{50, 60}, 11}, {{60, 70}, 3}, {{70, 80}, 12}, {{80, 90}, 4}
+			});
+		}
+
+		SECTION("PartialOverlap")
+		{
+			// IntervalTreeB's intervals partially overlap IntervalTree's intervals.
+			// Different data: existing intervals are trimmed to make room for the merged ones.
+			// IntervalTree has: [20,30),1  [40,50),2  [60,70),3  [80,90),4
+			TIntervalTree<int> IntervalTreeB;
+			Insert(IntervalTreeB, 25, 10, 10); // [25,35) overlaps [20,30),1 from the right
+			Insert(IntervalTreeB, 45, 10, 11); // [45,55) overlaps [40,50),2 from the right
+			Insert(IntervalTreeB, 55, 10, 12); // [55,65) overlaps [60,70),3 from the left
+			Insert(IntervalTreeB, 75, 20, 13); // [75,95) straddles [80,90),4 entirely
+			Validate(IntervalTreeB, {{{25, 35}, 10}, {{45, 55}, 11}, {{55, 65}, 12}, {{75, 95}, 13}});
+
+			IntervalTree.Merge(IntervalTreeB);
+			// [20,30),1 trimmed to [20,25),1 by [25,35),10
+			// [40,50),2 trimmed to [40,45),2 by [45,55),11
+			// [60,70),3 trimmed to [65,70),3 by [55,65),12
+			// [80,90),4 deleted by [75,95),13
+			Validate(IntervalTree, {
+				{{20, 25}, 1}, {{25, 35}, 10}, {{40, 45}, 2}, {{45, 55}, 11},
+				{{55, 65}, 12}, {{65, 70}, 3}, {{75, 95}, 13}
+			});
 		}
 	}
 
@@ -197,15 +406,191 @@ TEST_CASE("IntervalTree")
 
 			for (const FInterval& FreeInterval : Intervals.Free)
 			{
-				REQUIRE(nullptr == IntervalTree.Get(Ptr(FreeInterval.Start), FreeInterval.Size()));
-				REQUIRE(nullptr == IntervalTree.Get(Ptr(FreeInterval.Start), 1));
-				REQUIRE(nullptr == IntervalTree.Get(Ptr(FreeInterval.End - 1), 1));
+				REQUIRE(!HasOverlap(IntervalTree, FreeInterval.Start, FreeInterval.Size()));
+				REQUIRE(!HasOverlap(IntervalTree, FreeInterval.Start, 1));
+				REQUIRE(!HasOverlap(IntervalTree, FreeInterval.End - 1, 1));
 			}
 			for (const FInterval& UsedInterval : Intervals.Used)
 			{
-				REQUIRE(nullptr != IntervalTree.Get(Ptr(UsedInterval.Start), UsedInterval.Size()));
-				REQUIRE(nullptr != IntervalTree.Get(Ptr(UsedInterval.Start), 1));
-				REQUIRE(nullptr != IntervalTree.Get(Ptr(UsedInterval.End - 1), 1));
+				REQUIRE(HasOverlap(IntervalTree, UsedInterval.Start, UsedInterval.Size()));
+				REQUIRE(HasOverlap(IntervalTree, UsedInterval.Start, 1));
+				REQUIRE(HasOverlap(IntervalTree, UsedInterval.End - 1, 1));
+			}
+		}
+	}
+
+	SECTION("Soak.Overlapping")
+	{
+		// Validates the tree invariants after inserting intervals with overlaps.
+		auto ValidateOrdered = [&] AUTORTFM_DISABLE (TIntervalTree<int>& Tree)
+		{
+			uintptr_t PrevEnd = 0;
+			for (TIntervalTree<int>::FIntervalAndData Entry : Tree)
+			{
+				REQUIRE(Entry.Interval.Start >= PrevEnd);
+				REQUIRE(Entry.Interval.End > Entry.Interval.Start);
+				REQUIRE(HasOverlap(Tree, Entry.Interval.Start, Entry.Interval.Size()));
+				REQUIRE((*Tree.Get(Ptr(Entry.Interval.Start), Entry.Interval.Size()).begin()).Data == Entry.Data);
+				PrevEnd = Entry.Interval.End;
+			}
+		};
+
+		for (int Pass = 0; Pass < 32; Pass++)
+		{
+			IntervalTree.Reset();
+			Catch::SimplePcg32 Rand(Pass + 1000);
+
+			// Insert random overlapping intervals with varying data.
+			for (int I = 0; I < 500; I++)
+			{
+				uintptr_t Start = 1 + (Rand() % 1000);
+				uintptr_t Size = 1 + (Rand() % 50);
+				int Data = static_cast<int>(Rand() % 5);
+				Insert(IntervalTree, Start, Size, Data);
+			}
+
+			ValidateOrdered(IntervalTree);
+		}
+	}
+
+	SECTION("Get")
+	{
+		// Helper: collect Get() results into a vector for easy comparison.
+		auto GetAll = [Ptr](const TIntervalTree<int>& Tree, uintptr_t Address, uintptr_t Size) -> std::vector<FIntervalAndData>
+		{
+			std::vector<FIntervalAndData> Results;
+			for (TIntervalTree<int>::FIntervalAndData Entry : Tree.Get(Ptr(Address), Size))
+			{
+				Results.push_back(FIntervalAndData{Entry.Interval, Entry.Data});
+			}
+			return Results;
+		};
+
+		// Empty tree — empty range.
+		REQUIRE(GetAll(IntervalTree, 40, 10).empty());
+
+		// Populate: [20,30),1  [40,50),2  [60,70),3  [80,90),4
+		Insert(IntervalTree, 20, 10, 1);
+		Insert(IntervalTree, 40, 10, 2);
+		Insert(IntervalTree, 60, 10, 3);
+		Insert(IntervalTree, 80, 10, 4);
+
+		// No overlap — query falls in a gap.
+		REQUIRE(GetAll(IntervalTree, 30, 10).empty()); // [30,40)
+		REQUIRE(GetAll(IntervalTree, 10, 10).empty()); // [10,20)
+		REQUIRE(GetAll(IntervalTree, 90, 10).empty()); // [90,100)
+
+		// Single exact match.
+		{
+			std::vector<FIntervalAndData> Results = GetAll(IntervalTree, 40, 10);
+			REQUIRE(Results.size() == 1);
+			REQUIRE(Results[0].Interval.Start == 40);
+			REQUIRE(Results[0].Interval.End == 50);
+			REQUIRE(Results[0].Data == 2);
+		}
+
+		// Query starts before a stored interval but overlaps it (SOL-8992 case).
+		{
+			std::vector<FIntervalAndData> Results = GetAll(IntervalTree, 35, 10); // [35,45) overlaps [40,50)
+			REQUIRE(Results.size() == 1);
+			REQUIRE(Results[0].Interval.Start == 40);
+			REQUIRE(Results[0].Data == 2);
+		}
+
+		// Query ends after a stored interval.
+		{
+			std::vector<FIntervalAndData> Results = GetAll(IntervalTree, 45, 10); // [45,55) overlaps [40,50)
+			REQUIRE(Results.size() == 1);
+			REQUIRE(Results[0].Interval.Start == 40);
+			REQUIRE(Results[0].Data == 2);
+		}
+
+		// Query spans multiple intervals.
+		{
+			std::vector<FIntervalAndData> Results = GetAll(IntervalTree, 25, 50); // [25,75) overlaps [20,30), [40,50), [60,70)
+			REQUIRE(Results.size() == 3);
+			REQUIRE(Results[0].Interval.Start == 20);
+			REQUIRE(Results[0].Data == 1);
+			REQUIRE(Results[1].Interval.Start == 40);
+			REQUIRE(Results[1].Data == 2);
+			REQUIRE(Results[2].Interval.Start == 60);
+			REQUIRE(Results[2].Data == 3);
+		}
+
+		// Query spans all intervals.
+		{
+			std::vector<FIntervalAndData> Results = GetAll(IntervalTree, 0, 100); // [0,100) overlaps all four
+			REQUIRE(Results.size() == 4);
+		}
+
+		// Query contained within a single interval.
+		{
+			std::vector<FIntervalAndData> Results = GetAll(IntervalTree, 42, 3); // [42,45) inside [40,50)
+			REQUIRE(Results.size() == 1);
+			REQUIRE(Results[0].Data == 2);
+		}
+
+		// Predecessor-only overlap.
+		{
+			std::vector<FIntervalAndData> Results = GetAll(IntervalTree, 25, 10); // [25,35) overlaps [20,30) only
+			REQUIRE(Results.size() == 1);
+			REQUIRE(Results[0].Interval.Start == 20);
+			REQUIRE(Results[0].Data == 1);
+		}
+
+		// Adjacent but non-overlapping (half-open intervals don't overlap at endpoints).
+		REQUIRE(GetAll(IntervalTree, 30, 10).empty()); // [30,40) — between [20,30) and [40,50)
+		REQUIRE(GetAll(IntervalTree, 50, 10).empty()); // [50,60) — between [40,50) and [60,70)
+
+		// Single-byte query at interval start and end.
+		REQUIRE(GetAll(IntervalTree, 40, 1).size() == 1); // first byte of [40,50)
+		REQUIRE(GetAll(IntervalTree, 49, 1).size() == 1); // last byte of [40,50)
+		REQUIRE(GetAll(IntervalTree, 50, 1).empty());     // one past end
+
+		// Soak: random queries against BuildIntervals, brute-force verify.
+		for (int Pass = 0; Pass < 32; Pass++)
+		{
+			IntervalTree.Reset();
+			FTreeIntervals Intervals = BuildIntervals(/* Seed */ Pass + 2000, /* MaxUsed */ 500);
+
+			for (FInterval Interval : Intervals.Used)
+			{
+				Insert(IntervalTree, Interval.Start, Interval.Size(), 0);
+			}
+
+			// Collect the tree's actual intervals (which may have been merged).
+			std::vector<FInterval> TreeIntervals;
+			for (TIntervalTree<int>::FIntervalAndData Entry : IntervalTree)
+			{
+				TreeIntervals.push_back(Entry.Interval);
+			}
+
+			Catch::SimplePcg32 Rand(Pass + 3000);
+			for (int Q = 0; Q < 200; Q++)
+			{
+				uintptr_t QStart = Rand() % 0x200000;
+				uintptr_t QSize = 1 + Rand() % 0x100000;
+				uintptr_t QEnd = QStart + QSize;
+
+				// Brute-force: count tree intervals that overlap [QStart, QEnd).
+				size_t ExpectedCount = 0;
+				for (const FInterval& Interval : TreeIntervals)
+				{
+					if (Interval.Start < QEnd && QStart < Interval.End)
+					{
+						ExpectedCount++;
+					}
+				}
+
+				std::vector<FIntervalAndData> Results = GetAll(IntervalTree, QStart, QSize);
+				REQUIRE(Results.size() == ExpectedCount);
+
+				// Verify each result actually overlaps the query.
+				for (const FIntervalAndData& Entry : Results)
+				{
+					REQUIRE(Entry.Interval.Start < QEnd);
+					REQUIRE(QStart < Entry.Interval.End);
+				}
 			}
 		}
 	}
diff --git a/Engine/Source/Programs/AutoRTFMTests/Private/Open.cpp b/Engine/Source/Programs/AutoRTFMTests/Private/Open.cpp
--- a/Engine/Source/Programs/AutoRTFMTests/Private/Open.cpp
+++ b/Engine/Source/Programs/AutoRTFMTests/Private/Open.cpp
@@ -854,6 +854,43 @@ TEST_CASE("Open.Collision")
         });
     }
 
+    SECTION("Multiple closed writes, single spanning open write")
+    {
+        // Write to Array[0] and Array[2] transactionally (separate intervals),
+        // then do an open write spanning the whole array.
+        int Array[4] = {};
+        AutoRTFM::Transact([&]
+        {
+            Array[0] = 1;
+            Array[2] = 2;
+            AutoRTFM::Open([&]
+            {
+                SafeToAssignTo = 10;
+                WRITE_THAT_CAUSES_COLLISION(memset(Array, 0, sizeof(Array)));
+                SafeToAssignTo = 20;
+            });
+        });
+
+        REQUIRE(WarningContext.HasWarningSubstring("conflicting transactional write(s)"));
+    }
+
+    SECTION("Partial overlap, open write starts before closed write")
+    {
+        // Write to Array[1..2] transactionally, then do an open write to Array[0..1]
+        // which starts before the closed interval (the SOL-8992 scenario).
+        int Array[4] = {};
+        AutoRTFM::Transact([&]
+        {
+            Array[1] = 42;
+            AutoRTFM::Open([&]
+            {
+                SafeToAssignTo = 10;
+                WRITE_THAT_CAUSES_COLLISION(memset(&Array[0], 0, 2 * sizeof(int))); // write to Array[0..1]
+                SafeToAssignTo = 20;
+            });
+        });
+    }
+
     if (LineThatPerformsConflictingWrite > 0)
     {
         REQUIRE(WarningContext.HasWarningSubstring(AutoRTFMTestUtils::kMemoryModifiedInOpenWarning));
diff --git a/Engine/Source/Runtime/AutoRTFM/Private/Interval.h b/Engine/Source/Runtime/AutoRTFM/Private/Interval.h
--- a/Engine/Source/Runtime/AutoRTFM/Private/Interval.h
+++ b/Engine/Source/Runtime/AutoRTFM/Private/Interval.h
@@ -13,6 +13,9 @@ namespace AutoRTFM
 // A half-open address interval [Start, End).
 struct FInterval final
 {
+	// Default constructor - leaves Start and End uninitialized.
+	FInterval() = default;
+
 	// Constructs an interval with explicit start and end addresses.
 	constexpr FInterval(uintptr_t Start, uintptr_t End) : Start{Start}, End{End} {}
 
diff --git a/Engine/Source/Runtime/AutoRTFM/Private/IntervalTree.h b/Engine/Source/Runtime/AutoRTFM/Private/IntervalTree.h
--- a/Engine/Source/Runtime/AutoRTFM/Private/IntervalTree.h
+++ b/Engine/Source/Runtime/AutoRTFM/Private/IntervalTree.h
@@ -16,10 +16,12 @@ namespace AutoRTFM
 {
 
 // TIntervalTree tracks a set of non-overlapping memory intervals, each associated with a
-// data value of type DataTypeIn, using a red-black tree. Intervals are stored as [Start, End)
-// and intervals that touch at their endpoints are merged into a single interval on insert,
-// retaining the data of the existing interval.
+// data value of type DataTypeIn, using a red-black tree. Intervals are stored as [Start, End).
+// On insert, the new interval takes ownership of its address range:
+//   - Overlapping or adjacent intervals with the same data are merged.
+//   - Overlapping intervals with different data are trimmed, split, or deleted.
 template <typename DataTypeIn, Allocator Allocator = FAllocator>
+requires std::equality_comparable<DataTypeIn>
 struct AUTORTFM_DISABLE_NO_SANITIZE TIntervalTree final
 {
 private:
@@ -34,28 +36,21 @@ public:
 	TIntervalTree& operator=(const TIntervalTree&) = delete;
 
 	// Inserts the interval [Address, Address+Size) with the given data into the tree.
-	// If the new interval is adjacent to an existing interval, the two are merged and
-	// the existing interval's data is retained.
-	// Returns false if the interval overlaps (but is not adjacent to) an existing interval.
-	bool Insert(const void* const Address, const size_t Size, DataType&& Data)
+	// The new interval takes ownership of its address range:
+	//   - Overlapping or adjacent same-data intervals are merged into one.
+	//   - Overlapping different-data intervals are trimmed, split, or deleted to make room.
+	void Insert(const void* const Address, const size_t Size, DataType&& Data)
 	{
-		return Insert(FInterval{Address, Size}, std::move(Data));
+		Insert(FInterval{Address, Size}, std::move(Data));
 	}
 
 	// Inserts the interval [Address, Address+Size) with the given data into the tree.
-	// If the new interval is adjacent to an existing interval, the two are merged and
-	// the existing interval's data is retained.
-	// Returns false if the interval overlaps (but is not adjacent to) an existing interval.
-	bool Insert(const void* const Address, const size_t Size, const DataType& Data)
+	// The new interval takes ownership of its address range:
+	//   - Overlapping or adjacent same-data intervals are merged into one.
+	//   - Overlapping different-data intervals are trimmed, split, or deleted to make room.
+	void Insert(const void* const Address, const size_t Size, const DataType& Data)
 	{
-		return Insert(FInterval{Address, Size}, Data);
-	}
-
-	// Returns a pointer to the data associated with the stored interval that contains
-	// [Address, Address+Size), or nullptr if no such interval exists.
-	const DataType* Get(const void* const Address, const size_t Size) const
-	{
-		return Query(Address, Size);
+		Insert(FInterval{Address, Size}, Data);
 	}
 
 	// Returns true if no intervals are stored.
@@ -71,7 +66,9 @@ public:
 		Nodes.Reset();
 	}
 
-	// Inserts all intervals from Other into this tree.
+	// Inserts all intervals from Other into this tree, applying the same ownership
+	// semantics as Insert: overlapping/adjacent same-data intervals merge, and
+	// overlapping different-data intervals are trimmed, split, or deleted.
 	void Merge(const TIntervalTree& Other)
 	{
 		if (Other.IsEmpty())
@@ -130,7 +127,6 @@ public:
 		}
 
 		bool operator==(const FIterator& Other) const = default;
-		bool operator!=(const FIterator& Other) const = default;
 
 	private:
 		friend struct TIntervalTree;
@@ -151,6 +147,115 @@ public:
 		return FIterator(this, InvalidNode);
 	}
 
+	// A lightweight range representing all stored intervals that overlap a query.
+	struct FOverlapRange
+	{
+		FIterator Begin;
+		FIterator End;
+		FIterator begin() const
+		{
+			return Begin;
+		}
+		FIterator end() const
+		{
+			return End;
+		}
+	};
+
+	// Returns true if any stored interval overlaps [Address, Address+Size).
+	// Faster than Get() when you only need an existence check: single tree traversal,
+	// no iterator construction, no end-node lookup.
+	UE_AUTORTFM_FORCENOINLINE bool Overlaps(const void* const Address, const size_t Size) const
+	{
+		if (InvalidNode == Root)
+		{
+			return false;
+		}
+
+		const uintptr_t QueryStart = reinterpret_cast<uintptr_t>(Address);
+		const uintptr_t QueryEnd = QueryStart + Size;
+
+		const FInterval* const Intervals = Nodes.Ptrs.Interval;
+		const NodeIndex* const Lefts = Nodes.Ptrs.Left;
+		const NodeIndex* const Rights = Nodes.Ptrs.Right;
+
+		NodeIndex LowerBoundNode = InvalidNode;
+		NodeIndex PredecessorNode = InvalidNode;
+		{
+			NodeIndex Current = Root;
+			while (Current != InvalidNode)
+			{
+				if (Intervals[Current].Start >= QueryStart)
+				{
+					LowerBoundNode = Current;
+					Current = Lefts[Current];
+				}
+				else
+				{
+					PredecessorNode = Current;
+					Current = Rights[Current];
+				}
+			}
+		}
+
+		if (AUTORTFM_UNLIKELY(PredecessorNode != InvalidNode && Intervals[PredecessorNode].End > QueryStart))
+		{
+			return true;
+		}
+
+		return AUTORTFM_UNLIKELY(LowerBoundNode != InvalidNode && Intervals[LowerBoundNode].Start < QueryEnd);
+	}
+
+	// Returns all stored intervals overlapping [Address, Address+Size), in ascending
+	// order of start address.
+	UE_AUTORTFM_FORCENOINLINE FOverlapRange Get(const void* const Address, const size_t Size) const
+	{
+		if (AUTORTFM_UNLIKELY(InvalidNode == Root))
+		{
+			return {end(), end()};
+		}
+
+		const uintptr_t QueryStart = reinterpret_cast<uintptr_t>(Address);
+		const uintptr_t QueryEnd = QueryStart + Size;
+
+		const FInterval* const Intervals = Nodes.Ptrs.Interval;
+		const NodeIndex* const Lefts = Nodes.Ptrs.Left;
+		const NodeIndex* const Rights = Nodes.Ptrs.Right;
+
+		NodeIndex LowerBoundNode = InvalidNode;
+		NodeIndex PredecessorNode = InvalidNode;
+		{
+			NodeIndex Current = Root;
+			while (Current != InvalidNode)
+			{
+				if (Intervals[Current].Start >= QueryStart)
+				{
+					LowerBoundNode = Current;
+					Current = Lefts[Current];
+				}
+				else
+				{
+					PredecessorNode = Current;
+					Current = Rights[Current];
+				}
+			}
+		}
+
+		NodeIndex BeginNode = LowerBoundNode;
+		if (PredecessorNode != InvalidNode && Intervals[PredecessorNode].End > QueryStart)
+		{
+			BeginNode = PredecessorNode;
+		}
+
+		if (BeginNode == InvalidNode || Intervals[BeginNode].Start >= QueryEnd)
+		{
+			return {end(), end()};
+		}
+
+		const NodeIndex EndNode = LowerBound(QueryEnd);
+		return {FIterator(this, BeginNode), FIterator(this, EndNode)};
+	}
+
 private:
 	friend struct FIterator;
 
@@ -290,10 +395,10 @@ private:
 			}
 		}
 
-		EColor Color(size_t Index) const
+		EColor Color(NodeIndex Index) const
 		{
 			ColorWord const Bit = static_cast<ColorWord>(1) << (Index & (NumColorWordBits - 1));
-			ColorWord& Word = Ptrs.Color[Index / NumColorWordBits];
+			const ColorWord Word = Ptrs.Color[Index / NumColorWordBits];
 			return (Word & Bit) ? EColor::Black : EColor::Red;
 		}
 
@@ -366,8 +471,6 @@ private:
 	NodeIndex Root = InvalidNode;
 	FNodes Nodes;
 
-	static constexpr bool bExtraDebugging = false;
-
 	// Returns the index of the next sequential node after Node.
 	NodeIndex Next(NodeIndex Node) const
 	{
@@ -401,6 +504,396 @@ private:
 		return Node;
 	}
 
+	// Returns the rightmost node whose Start < SearchStart, or InvalidNode.
+	NodeIndex Predecessor(uintptr_t SearchStart) const
+	{
+		NodeIndex Result = InvalidNode;
+		NodeIndex Current = Root;
+		while (Current != InvalidNode)
+		{
+			if (Nodes.Interval(Current).Start < SearchStart)
+			{
+				Result = Current;
+				Current = Nodes.Right(Current);
+			}
+			else
+			{
+				Current = Nodes.Left(Current);
+			}
+		}
+		return Result;
+	}
+
+	// Returns the leftmost node whose Start >= SearchStart, or InvalidNode.
+	NodeIndex LowerBound(uintptr_t SearchStart) const
+	{
+		NodeIndex Result = InvalidNode;
+		NodeIndex Current = Root;
+		while (Current != InvalidNode)
+		{
+			if (Nodes.Interval(Current).Start >= SearchStart)
+			{
+				Result = Current;
+				Current = Nodes.Left(Current);
+			}
+			else
+			{
+				Current = Nodes.Right(Current);
+			}
+		}
+		return Result;
+	}
+
+	// Left-rotates around X:  X's right child Y takes X's place.
+	void RotateLeft(NodeIndex X)
+	{
+		const NodeIndex Y = Nodes.Right(X);
+		AUTORTFM_ASSERT(Y != InvalidNode);
+
+		Nodes.Right(X) = Nodes.Left(Y);
+		if (Nodes.Left(Y) != InvalidNode)
+		{
+			Nodes.Parent(Nodes.Left(Y)) = X;
+		}
+
+		Nodes.Parent(Y) = Nodes.Parent(X);
+		if (Nodes.Parent(X) == InvalidNode)
+		{
+			Root = Y;
+		}
+		else if (X == Nodes.Left(Nodes.Parent(X)))
+		{
+			Nodes.Left(Nodes.Parent(X)) = Y;
+		}
+		else
+		{
+			Nodes.Right(Nodes.Parent(X)) = Y;
+		}
+
+		Nodes.Left(Y) = X;
+		Nodes.Parent(X) = Y;
+	}
+
+	// Right-rotates around X:  X's left child Y takes X's place.
+	void RotateRight(NodeIndex X)
+	{
+		const NodeIndex Y = Nodes.Left(X);
+		AUTORTFM_ASSERT(Y != InvalidNode);
+
+		Nodes.Left(X) = Nodes.Right(Y);
+		if (Nodes.Right(Y) != InvalidNode)
+		{
+			Nodes.Parent(Nodes.Right(Y)) = X;
+		}
+
+		Nodes.Parent(Y) = Nodes.Parent(X);
+		if (Nodes.Parent(X) == InvalidNode)
+		{
+			Root = Y;
+		}
+		else if (X == Nodes.Right(Nodes.Parent(X)))
+		{
+			Nodes.Right(Nodes.Parent(X)) = Y;
+		}
+		else
+		{
+			Nodes.Left(Nodes.Parent(X)) = Y;
+		}
+
+		Nodes.Right(Y) = X;
+		Nodes.Parent(X) = Y;
+	}
+
+	// Removes node Z from the red-black tree.
+	// The node's slot in the SoA array is left as a dead entry (not reclaimed).
+	UE_AUTORTFM_FORCENOINLINE void Remove(NodeIndex Z)
+	{
+		// If Z has two children, swap it with its in-order successor S.
+		// S is the leftmost node of Z's right subtree, so it has at most one
+		// child (a right child).  We copy S's data into Z and then delete S
+		// from the tree structure instead.
+		if (Nodes.Left(Z) != InvalidNode && Nodes.Right(Z) != InvalidNode)
+		{
+			const NodeIndex S = LeftMostDescendant(Nodes.Right(Z));
+			Nodes.Interval(Z) = Nodes.Interval(S);
+			Nodes.Data(Z) = std::move(Nodes.Data(S));
+			Z = S;
+		}
+
+		// Z now has at most one non-null child.
+		const NodeIndex C = (Nodes.Left(Z) != InvalidNode) ? Nodes.Left(Z) : Nodes.Right(Z);
+		const NodeIndex P = Nodes.Parent(Z);
+
+		// Splice Z out of the tree.
+		if (C != InvalidNode)
+		{
+			Nodes.Parent(C) = P;
+		}
+		if (P == InvalidNode)
+		{
+			Root = C;
+		}
+		else if (Z == Nodes.Left(P))
+		{
+			Nodes.Left(P) = C;
+		}
+		else
+		{
+			Nodes.Right(P) = C;
+		}
+
+		// Color fixup.
+		if (Nodes.Color(Z) == EColor::Red)
+		{
+			// Removing a red node never changes black height.
+			return;
+		}
+		if (C != InvalidNode && Nodes.Color(C) == EColor::Red)
+		{
+			// The replacement is red: recolor it black to restore black height.
+			Nodes.SetColor(C, EColor::Black);
+			return;
+		}
+
+		// C is "double-black".  Walk up the tree applying the standard
+		// six-case RB fixup until the extra black is absorbed.
+		NodeIndex X = C;        // double-black node (may be InvalidNode = null)
+		NodeIndex XParent = P;  // parent of X
+
+		while (XParent != InvalidNode && (X == InvalidNode || Nodes.Color(X) == EColor::Black))
+		{
+			const bool bXIsLeft = (Nodes.Left(XParent) == X);
+			NodeIndex W = bXIsLeft ? Nodes.Right(XParent) : Nodes.Left(XParent);  // sibling
+
+			// Case 1: sibling W is red — rotate to make W black.
+			if (W != InvalidNode && Nodes.Color(W) == EColor::Red)
+			{
+				Nodes.SetColor(W, EColor::Black);
+				Nodes.SetColor(XParent, EColor::Red);
+				if (bXIsLeft)
+				{
+					RotateLeft(XParent);
+				}
+				else
+				{
+					RotateRight(XParent);
+				}
+				W = bXIsLeft ? Nodes.Right(XParent) : Nodes.Left(XParent);
+			}
+
+			// W is now black (or null).
+			const NodeIndex WFar = (W != InvalidNode) ? (bXIsLeft ? Nodes.Right(W) : Nodes.Left(W)) : InvalidNode;
+			const NodeIndex WNear = (W != InvalidNode) ? (bXIsLeft ? Nodes.Left(W) : Nodes.Right(W)) : InvalidNode;
+			const bool bWFarBlack = (WFar == InvalidNode || Nodes.Color(WFar) == EColor::Black);
+			const bool bWNearBlack = (WNear == InvalidNode || Nodes.Color(WNear) == EColor::Black);
+
+			if (bWFarBlack && bWNearBlack)
+			{
+				// Case 2: both of W's children are black — push the double-black up.
+				if (W != InvalidNode)
+				{
+					Nodes.SetColor(W, EColor::Red);
+				}
+				if (Nodes.Color(XParent) == EColor::Red)
+				{
+					Nodes.SetColor(XParent, EColor::Black);
+					return;  // resolved
+				}
+				X = XParent;
+				XParent = Nodes.Parent(XParent);
+				continue;
+			}
+
+			if (bWFarBlack)
+			{
+				// Case 3: near child is red, far child is black — rotate W to
+				// make the far child red, then fall through to case 4.
+				if (WNear != InvalidNode)
+				{
+					Nodes.SetColor(WNear, EColor::Black);
+				}
+				if (W != InvalidNode)
+				{
+					Nodes.SetColor(W, EColor::Red);
+				}
+				if (bXIsLeft)
+				{
+					RotateRight(W);
+				}
+				else
+				{
+					RotateLeft(W);
+				}
+				W = bXIsLeft ? Nodes.Right(XParent) : Nodes.Left(XParent);
+			}
+
+			// Case 4: far child is red — final rotation resolves the double-black.
+			const NodeIndex NewWFar = (W != InvalidNode) ? (bXIsLeft ? Nodes.Right(W) : Nodes.Left(W)) : InvalidNode;
+			Nodes.SetColor(W, Nodes.Color(XParent));
+			Nodes.SetColor(XParent, EColor::Black);
+			if (NewWFar != InvalidNode)
+			{
+				Nodes.SetColor(NewWFar, EColor::Black);
+			}
+			if (bXIsLeft)
+			{
+				RotateLeft(XParent);
+			}
+			else
+			{
+				RotateRight(XParent);
+			}
+			return;  // case 4 always terminates
+		}
+
+		// X has reached the root or is red: color it black to absorb the extra black.
+		if (X != InvalidNode)
+		{
+			Nodes.SetColor(X, EColor::Black);
+		}
+	}
+
+	// Inserts NewInterval with NewData into the tree, applying the ownership semantics:
+	//   Phase A — extend MergedStart left by consuming adjacent/overlapping same-data predecessors;
+	//             trim or split a different-data predecessor that overlaps.
+	//   Phase B — scan forward from MergedStart: consume same-data nodes (extending MergedEnd),
+	//             delete fully-covered different-data nodes, trim a right-partial overlap.
+	//   Phase C — SimpleInsert the effective [MergedStart, MergedEnd) interval.
+	//   Phase D — re-insert any deferred right-remainder created by a split or right trim.
+	UE_AUTORTFM_FORCENOINLINE void Insert(FInterval NewInterval, DataType NewData)
+	{
+		uintptr_t MergedStart = NewInterval.Start;
+		uintptr_t MergedEnd = NewInterval.End;
+
+		// Storage for an optional deferred re-insertion (at most one can arise).
+		// Placed in aligned raw storage to avoid requiring DataType to be default-constructible.
+		bool bHasRemainder = false;
+		FInterval RemainderInterval{uintptr_t{0}, uintptr_t{0}};
+		alignas(DataType) std::byte RemainderDataBuf[sizeof(DataType)];
+		auto SetRemainder = [&](FInterval Interval, const DataType& Data)
+		{
+			AUTORTFM_ASSERT(!bHasRemainder);
+			bHasRemainder = true;
+			RemainderInterval = Interval;
+			::new (RemainderDataBuf) DataType(Data);
+		};
+
+		// === Phase A: handle the predecessor of MergedStart ===
+		for (;;)
+		{
+			const NodeIndex Left = Predecessor(MergedStart);
+			if (Left == InvalidNode)
+			{
+				break;
+			}
+
+			const FInterval LeftInterval = Nodes.Interval(Left);  // value copy — safe across modifications
+			if (LeftInterval.End < MergedStart)
+			{
+				break;  // no contact
+			}
+
+			if (LeftInterval.End == MergedStart)
+			{
+				// Left-adjacent.
+				if (Nodes.Data(Left) == NewData)
+				{
+					MergedStart = LeftInterval.Start;
+					Remove(Left);
+					continue;  // re-check new predecessor
+				}
+				break;  // different data adjacent: leave Left as-is
+			}
+
+			// LeftInterval.End > MergedStart: Left overlaps from the left.
+			if (Nodes.Data(Left) == NewData)
+			{
+				MergedStart = LeftInterval.Start;
+				if (LeftInterval.End > MergedEnd)
+				{
+					MergedEnd = LeftInterval.End;
+				}  // extend right too
+				Remove(Left);
+				continue;
+			}
+
+			// Different data: trim Left's End in-place (End-only change; BST-safe).
+			if (LeftInterval.End > MergedEnd)
+			{
+				// Left straddles our whole interval — save its right portion for Phase D.
+				SetRemainder(FInterval{MergedEnd, LeftInterval.End}, Nodes.Data(Left));
+			}
+			Nodes.Interval(Left).End = MergedStart;
+			break;
+		}
+
+		// === Phase B: handle nodes starting at or after MergedStart ===
+		// Re-queries LowerBound each iteration because Remove() may rearrange the
+		// BST (swapping a node with its in-order successor), which can invalidate
+		// any previously-computed Next() pointer.
+		for (;;)
+		{
+			const NodeIndex Right = LowerBound(MergedStart);
+			if (Right == InvalidNode)
+			{
+				break;
+			}
+
+			const FInterval RightInterval = Nodes.Interval(Right);  // value copy
+			if (RightInterval.Start > MergedEnd)
+			{
+				break;  // past our interval
+			}
+
+			if (RightInterval.Start == MergedEnd)
+			{
+				// Right-adjacent.
+				if (Nodes.Data(Right) == NewData)
+				{
+					MergedEnd = RightInterval.End;
+					Remove(Right);
+				}
+				break;  // stop regardless (diff data adjacent: leave Right as-is)
+			}
+
+			// RightInterval.Start < MergedEnd: Right overlaps our effective interval.
+			if (Nodes.Data(Right) == NewData)
+			{
+				if (RightInterval.End > MergedEnd)
+				{
+					MergedEnd = RightInterval.End;
+				}
+				Remove(Right);
+				continue;
+			}
+
+			// Different data: clip Right.
+			if (RightInterval.End <= MergedEnd)
+			{
+				// Fully contained: new interval takes ownership — delete Right.
+				Remove(Right);
+				continue;
+			}
+
+			// Right extends past MergedEnd: keep [MergedEnd, RightInterval.End) with RData.
+			// We cannot change Right.Start in-place (BST order violation), so defer.
+			SetRemainder(FInterval{MergedEnd, RightInterval.End}, Nodes.Data(Right));
+			Remove(Right);
+			break;
+		}
+
+		// === Phase C: insert the (possibly extended) effective interval ===
+		SimpleInsert(FInterval{MergedStart, MergedEnd}, std::move(NewData));
+
+		// === Phase D: re-insert any deferred right remainder ===
+		if (bHasRemainder)
+		{
+			auto* RemainderData = reinterpret_cast<DataType*>(RemainderDataBuf);
+			SimpleInsert(RemainderInterval, std::move(*RemainderData));
+			RemainderData->~DataType();
+		}
+	}
+
 	UE_AUTORTFM_FORCEINLINE NodeIndex AddNode(FInterval NewInterval, DataType&& Data, NodeIndex Parent = InvalidNode)
 	{
 		return Nodes.Push({
@@ -411,13 +904,14 @@ private:
 		});
 	}
 
-	UE_AUTORTFM_FORCENOINLINE bool Insert(FInterval NewInterval, DataType Data)
+	// Fast-path insert: assumes [NewInterval] has no overlapping or adjacent intervals
+	// already in the tree.  Called by Insert() after it has cleared the way.
+	UE_AUTORTFM_FORCENOINLINE void SimpleInsert(FInterval NewInterval, DataType Data)
 	{
 		if (AUTORTFM_UNLIKELY(InvalidNode == Root))
 		{
-			AUTORTFM_ASSERT(Nodes.Count == 0);
 			Root = AddNode(NewInterval, std::move(Data));
-			return true;
+			return;
 		}
 
 		NodeIndex Current = Root;
@@ -426,22 +920,12 @@ private:
 		{
 			const FInterval Interval = Nodes.Interval(Current);
 
-			if (AUTORTFM_UNLIKELY((NewInterval.Start < Interval.End) && (Interval.Start < NewInterval.End)))
-			{
-				return false;
-			}
+			// SimpleInsert requires no pre-existing overlaps.
+			AUTORTFM_ASSERT(!((NewInterval.Start < Interval.End) && (Interval.Start < NewInterval.End)));
 
 			if (NewInterval.Start < Interval.Start)
 			{
-				if (AUTORTFM_UNLIKELY(NewInterval.End == Interval.Start))
-				{
-					AUTORTFM_ASSERT(NewInterval.Start < Interval.Start);
-
-					// We can just modify the existing node in place.
-					Nodes.Interval(Current).Start = NewInterval.Start;
-					return true;
-				}
-				else if (InvalidNode == Nodes.Left(Current))
+				if (InvalidNode == Nodes.Left(Current))
 				{
 					const NodeIndex Index = AddNode(NewInterval, std::move(Data), Current);
 					Nodes.Left(Current) = Index;
@@ -453,15 +937,7 @@ private:
 			}
 			else
 			{
-				if (AUTORTFM_UNLIKELY(NewInterval.Start == Interval.End))
-				{
-					AUTORTFM_ASSERT(NewInterval.End > Interval.End);
-
-					// We can just modify the existing node in place.
-					Nodes.Interval(Current).End = NewInterval.End;
-					return true;
-				}
-				else if (InvalidNode == Nodes.Right(Current))
+				if (InvalidNode == Nodes.Right(Current))
 				{
 					const NodeIndex Index = AddNode(NewInterval, std::move(Data), Current);
 					Nodes.Right(Current) = Index;
@@ -617,86 +1093,6 @@ private:
 
 			break;
 		}
-
-		AssertStructureIsOk();
-
-		return true;
-	}
-
-	UE_AUTORTFM_FORCENOINLINE const DataType* Query(const void* const Address, const size_t Size) const
-	{
-		if (AUTORTFM_UNLIKELY(InvalidNode == Root))
-		{
-			return nullptr;
-		}
-
-		const FInterval QueryInterval{Address, Size};
-
-		NodeIndex Current = Root;
-
-		do
-		{
-			const FInterval Interval = Nodes.Interval(Current);
-
-			if ((QueryInterval.Start < Interval.End) && (Interval.Start < QueryInterval.End))
-			{
-				return &Nodes.Data(Current);
-			}
-
-			const NodeIndex* Next = (QueryInterval.Start < Interval.End) ? Nodes.Ptrs.Left : Nodes.Ptrs.Right;
-			Current = Next[Current];
-		} while (InvalidNode != Current);
-
-		return nullptr;
-	}
-
-	UE_AUTORTFM_FORCEINLINE void AssertStructureIsOk() const
-	{
-		if constexpr (bExtraDebugging)
-		{
-			if (InvalidNode != Root)
-			{
-				AssertNodeIsOk(Root);
-			}
-		}
-	}
-
-	UE_AUTORTFM_FORCENOINLINE void AssertNodeIsOk(NodeIndex Index) const
-	{
-		// We need to use recursion to check because we cannot have any
-		// allocations within the checker itself!
-		if constexpr (bExtraDebugging)
-		{
-			AUTORTFM_ASSERT(InvalidNode != Index);
-			AUTORTFM_ASSERT(Index < Nodes.Count);
-
-			const NodeIndex Parent = Nodes.Parent(Index);
-			const NodeIndex Left = Nodes.Left(Index);
-			const NodeIndex Right = Nodes.Right(Index);
-
-			if (InvalidNode == Parent)
-			{
-				AUTORTFM_ASSERT(Root == Index);
-			}
-			else
-			{
-				AUTORTFM_ASSERT(Nodes.Color(Parent) == EColor::Black || Nodes.Color(Index) == EColor::Black);
-				AUTORTFM_ASSERT((Nodes.Left(Parent) == Index) ^ (Nodes.Right(Parent) == Index));
-			}
-
-			AUTORTFM_ASSERT(Left != Index);
-			AUTORTFM_ASSERT(Right != Index);
-
-			if (InvalidNode != Left)
-			{
-				AssertNodeIsOk(Left);
-			}
-
-			if (InvalidNode != Right)
-			{
-				AssertNodeIsOk(Right);
-			}
-		}
 	}
 };
 
diff --git a/Engine/Source/Runtime/AutoRTFM/Private/Sanitizer.cpp b/Engine/Source/Runtime/AutoRTFM/Private/Sanitizer.cpp
--- a/Engine/Source/Runtime/AutoRTFM/Private/Sanitizer.cpp
+++ b/Engine/Source/Runtime/AutoRTFM/Private/Sanitizer.cpp
@@ -53,17 +53,29 @@ struct AUTORTFM_DISABLE_NO_SANITIZE FOSAllocator
 };
 static_assert(Allocator<FOSAllocator>);
 
+// Maximum number of closed write entries stored in a single conflict report.
+// If more are found, the report will note how many were omitted.
+static constexpr size_t MaxReportedClosedWrites = 4;
+
 struct AUTORTFM_DISABLE_NO_SANITIZE FSanitizer
 {
 	using FCallstackTree = TCallstackTree<FOSAllocator>;
 
 	struct FConflictingWrite
 	{
-		autortfm_log_severity const Severity;
-		void* const Ptr;
-		size_t const Size;
-		FCallstackTree::FCallstack const OpenWriteCallstack;
-		FCallstackTree::FCallstack const ClosedWriteCallstack;
+		struct FClosedWriteEntry
+		{
+			FInterval Interval;
+			FCallstackTree::FCallstack Callstack;
+		};
+
+		autortfm_log_severity Severity;
+		void* Ptr;
+		size_t Size;
+		FCallstackTree::FCallstack OpenWriteCallstack;
+		FClosedWriteEntry ClosedWrites[MaxReportedClosedWrites];
+		size_t NumClosedWrites = 0;    // entries stored in ClosedWrites[]
+		size_t TotalClosedWrites = 0;  // total found (may exceed MaxReportedClosedWrites)
 
 		// Report the conflicting write, aborting the application if severity is autortfm_log_fatal.
 		// As logging the conflicting write may call back into autortfm_sanitizer_open_write(), this
@@ -104,21 +116,6 @@ private:
 		TClosedWriteTree ClosedWrites;
 	};
 
-	struct FCallstackPair
-	{
-		FCallstackTree::FHandle Open;
-		FCallstackTree::FHandle Closed;
-
-		struct FHasher
-		{
-			size_t operator()(FCallstackPair Pair) const
-			{
-				return (static_cast<size_t>(Pair.Open) << 32) | static_cast<size_t>(Pair.Closed);
-			}
-		};
-		bool operator==(const FCallstackPair&) const = default;
-	};
-
 	static FSanitizer* Instance;
 
 	void PopTransaction(TransactionID ID);
@@ -127,7 +124,7 @@ private:
 
 	FTransaction* CurrentTransaction = nullptr;
 	FCallstackTree CallstackTree;
-	std::unordered_set<FCallstackPair, FCallstackPair::FHasher> ReportedCallstacks;
+	std::unordered_set<FCallstackTree::FHandle> ReportedCallstacks;
 	size_t NumIssues = 0;
 };
 
@@ -256,34 +253,62 @@ void FSanitizer::AbortTransaction(TransactionID ID)
 		return nullptr;
 	}
 
+	// Fast check: does any closed write overlap the open write range?
+	// This avoids constructing FClosedWriteEntry on the stack for the common no-collision case.
+	bool bHasOverlap = false;
+	for (FTransaction* Transaction = CurrentTransaction; Transaction && !bHasOverlap; Transaction = Transaction->Parent)
+	{
+		bHasOverlap = Transaction->ClosedWrites.Overlaps(Ptr, Size);
+	}
+
+	if (AUTORTFM_LIKELY(!bHasOverlap))
+	{
+		return nullptr;
+	}
+
+	// Collision found — gather all closed writes that overlap the open write range.
+	size_t NumClosedWrites = 0;
+	size_t TotalClosedWrites = 0;
+	FConflictingWrite::FClosedWriteEntry ClosedWrites[MaxReportedClosedWrites];
+
 	for (FTransaction* Transaction = CurrentTransaction; Transaction; Transaction = Transaction->Parent)
 	{
-		// TODO(SOL-8992): TIntervalTree::Get() may incorrectly return
-		// null if the queried interval starts outside of an interval added by
-		// TIntervalTree::Insert() but still overlaps an inserted interval.
-		// Fix this.
-		if (FCallstackTree::FHandle const* const ClosedWriteCallstack = Transaction->ClosedWrites.Get(Ptr, Size);
-			AUTORTFM_UNLIKELY(ClosedWriteCallstack))
+		for (TClosedWriteTree::FIntervalAndData Entry : Transaction->ClosedWrites.Get(Ptr, Size))
 		{
-			autortfm_log_severity const Severity = (Mode == Sanitizer::EMode::Error) ? autortfm_log_error : autortfm_log_warn;
-
-			FCallstackTree::FHandle const OpenWriteCallstack = GetCallstackHandle(ProgramCounter);
-			if (Severity == autortfm_log_warn)
+			if (NumClosedWrites < MaxReportedClosedWrites)
 			{
-				if (!ReportedCallstacks.emplace(OpenWriteCallstack, *ClosedWriteCallstack).second)
-				{
-					break;  // Already reported
-				}
+				ClosedWrites[NumClosedWrites] = {Entry.Interval, CallstackTree.Get(Entry.Data)};
+				NumClosedWrites++;
 			}
+			TotalClosedWrites++;
+		}
+	}
 
-			NumIssues++;
+	const autortfm_log_severity Severity = (Mode == Sanitizer::EMode::Error) ? autortfm_log_error : autortfm_log_warn;
+	FCallstackTree::FHandle const OpenHandle = GetCallstackHandle(ProgramCounter);
 
-			return New<FOSAllocator, FConflictingWrite>(
-				Severity, Ptr, Size, CallstackTree.Get(OpenWriteCallstack), CallstackTree.Get(*ClosedWriteCallstack));
+	if (Severity == autortfm_log_warn)
+	{
+		if (!ReportedCallstacks.emplace(OpenHandle).second)
+		{
+			return nullptr;  // Already reported for this open write site
 		}
 	}
 
-	return nullptr;
+	NumIssues++;
+
+	FConflictingWrite* Result = New<FOSAllocator, FConflictingWrite>();
+	Result->Severity = Severity;
+	Result->Ptr = Ptr;
+	Result->Size = Size;
+	Result->OpenWriteCallstack = CallstackTree.Get(OpenHandle);
+	Result->NumClosedWrites = NumClosedWrites;
+	Result->TotalClosedWrites = TotalClosedWrites;
+	for (size_t I = 0; I < NumClosedWrites; I++)
+	{
+		Result->ClosedWrites[I] = ClosedWrites[I];
+	}
+	return Result;
 }
 
 void FSanitizer::ClosedWrite(TransactionID ID, void* Ptr, size_t Size, void* ProgramCounter)
@@ -339,16 +364,56 @@ void FSanitizer::FConflictingWrite::Report() const
 	AUTORTFM_LOG_AT_SEVERITY(Severity, "This may lead to memory corruption if the transaction is aborted.");
 	AUTORTFM_LOG_AT_SEVERITY(Severity, "");
 
-	if (!ClosedWriteCallstack.IsEmpty())
+	if (TotalClosedWrites == 0)
 	{
-		AUTORTFM_LOG_AT_SEVERITY(Severity, "AutoRTFMSan: The last transactional write to this address was at:");
-		GExternAPI.LogCallstack(Severity, ClosedWriteCallstack.Num(), &ClosedWriteCallstack.Front());
-		AUTORTFM_LOG_AT_SEVERITY(Severity, "");
+		if (!bRecordClosedWriteCallstacks)
+		{
+			AUTORTFM_LOG_AT_SEVERITY(Severity, "AutoRTFMSan: Transactional write callstack recording is not enabled.");
+			AUTORTFM_LOG_AT_SEVERITY(Severity, "AutoRTFMSan: Enable to see callstack of the last closed write.");
+		}
 	}
-	else if (!bRecordClosedWriteCallstacks)
+	else if (TotalClosedWrites == 1)
 	{
-		AUTORTFM_LOG_AT_SEVERITY(Severity, "AutoRTFMSan: Transactional write callstack recording is not enabled.");
-		AUTORTFM_LOG_AT_SEVERITY(Severity, "AutoRTFMSan: Enable to see callstack of the last closed write.");
+		// Single closed write — use the familiar format.
+		const FClosedWriteEntry& Entry = ClosedWrites[0];
+		if (!Entry.Callstack.IsEmpty())
+		{
+			AUTORTFM_LOG_AT_SEVERITY(Severity, "AutoRTFMSan: The last transactional write to this address was at:");
+			GExternAPI.LogCallstack(Severity, Entry.Callstack.Num(), &Entry.Callstack.Front());
+			AUTORTFM_LOG_AT_SEVERITY(Severity, "");
+		}
+		else if (!bRecordClosedWriteCallstacks)
+		{
+			AUTORTFM_LOG_AT_SEVERITY(Severity, "AutoRTFMSan: Transactional write callstack recording is not enabled.");
+			AUTORTFM_LOG_AT_SEVERITY(Severity, "AutoRTFMSan: Enable to see callstack of the last closed write.");
+		}
+	}
+	else
+	{
+		// Multiple closed writes — numbered list.
+		AUTORTFM_LOG_AT_SEVERITY(Severity, "AutoRTFMSan: %zu conflicting transactional write(s):", TotalClosedWrites);
+		for (size_t I = 0; I < NumClosedWrites; I++)
+		{
+			const FClosedWriteEntry& Entry = ClosedWrites[I];
+			AUTORTFM_LOG_AT_SEVERITY(Severity, "");
+			AUTORTFM_LOG_AT_SEVERITY(
+				Severity, "  [%zu] 0x%zx bytes at %p:", I + 1, Entry.Interval.Size(), reinterpret_cast<void*>(Entry.Interval.Start));
+			if (!Entry.Callstack.IsEmpty())
+			{
+				GExternAPI.LogCallstack(Severity, Entry.Callstack.Num(), &Entry.Callstack.Front());
+			}
+			else if (!bRecordClosedWriteCallstacks)
+			{
+				AUTORTFM_LOG_AT_SEVERITY(Severity, "  <callstack recording not enabled>");
+			}
+		}
+		if (TotalClosedWrites > MaxReportedClosedWrites)
+		{
+			AUTORTFM_LOG_AT_SEVERITY(Severity, "");
+			AUTORTFM_LOG_AT_SEVERITY(
+				Severity, "  ... and %zu more conflicting write(s).", TotalClosedWrites - MaxReportedClosedWrites);
+		}
+		AUTORTFM_LOG_AT_SEVERITY(Severity, "");
 	}
 
 	AUTORTFM_LOG_AT_SEVERITY(Severity, "AutoRTFMSan: The open write was at:");
