[AutoRTFM] Address Bughawk complaint about hypothetical case where Result == 0 but GetNumEnsureFailures() > 0 leading to a return value of 0.

The code is now restructured so that we will definitely return an error code unless all conditions are successful.

This change will be landed in earlier streams at https://p4-swarm.epicgames.net/reviews/51787236

#Synced-CL 51784583

diff --git a/Engine/Source/Programs/AutoRTFMTests/Private/AutoRTFMTests.cpp b/Engine/Source/Programs/AutoRTFMTests/Private/AutoRTFMTests.cpp
--- a/Engine/Source/Programs/AutoRTFMTests/Private/AutoRTFMTests.cpp
+++ b/Engine/Source/Programs/AutoRTFMTests/Private/AutoRTFMTests.cpp
@@ -164,6 +164,7 @@ AUTORTFM_DISABLE int main(int ArgC, const char* ArgV[])
 	if (Result != 0)
 	{
 		FAIL("AutoRTFM tests failed with code: " << Result);
+		return Result;
 	}
 
 #if AUTORTFM_ENABLE_TEST_UTILS
@@ -176,5 +177,5 @@ AUTORTFM_DISABLE int main(int ArgC, const char* ArgV[])
 	}
 #endif // AUTORTFM_ENABLE_TEST_UTILS
 
-	return Result == 0 && FDebug::GetNumEnsureFailures() == 0 ? 0 : Result;
+	return FDebug::GetNumEnsureFailures() == 0 ? 0 : -1;
 }
