Change the `FPackageScopeMarker` to allow us to defer making a new package for tests that require a missing package.

Previously we wouldn't make sure to clear out any stale packages on disk so could cause random test failures.

#Synced-CL 51772012

diff --git a/Engine/Source/Programs/AutoRTFMTests/Private/UPackageTests.cpp b/Engine/Source/Programs/AutoRTFMTests/Private/UPackageTests.cpp
--- a/Engine/Source/Programs/AutoRTFMTests/Private/UPackageTests.cpp
+++ b/Engine/Source/Programs/AutoRTFMTests/Private/UPackageTests.cpp
@@ -109,11 +109,22 @@ static const TCHAR* ObjectName = TEXT("TestObject");
 
 // Based on `Engine\Private\Tests\Loading\AsyncLoadingTests_Shared.h`, we use similar logic
 // here to make a package that the loader will see and be able to actually load!
+enum class EShouldDeferPackageCreation : bool
+{
+	Yes = true,
+	No = false,
+};
+
 struct AUTORTFM_DISABLE FPackageScopedMaker final
 {
+	FString PackageName;
 	FString PackagePath;
+	EShouldDeferPackageCreation ShouldDeferPackageCreation;
 
-	FPackageScopedMaker(FString PackageName) : PackagePath(FPackageName::LongPackageNameToFilename(*PackageName, FPackageName::GetAssetPackageExtension()))
+	FPackageScopedMaker(FString InPackageName, EShouldDeferPackageCreation InShouldDeferPackageCreation = EShouldDeferPackageCreation::No) :
+		PackageName(InPackageName),
+		PackagePath(FPackageName::LongPackageNameToFilename(*InPackageName, FPackageName::GetAssetPackageExtension())),
+		ShouldDeferPackageCreation(InShouldDeferPackageCreation)
 	{
 		// We need to remove any previous package of the same name (could have occurred if a previous test ran segfaulted for instance).
 		if (FPackageName::DoesPackageExist(PackageName))
@@ -122,12 +133,26 @@ struct AUTORTFM_DISABLE FPackageScopedMaker final
 			REQUIRE(IPlatformFile::GetPlatformPhysical().DeleteFile(*PackagePath));
 		}
 
+		if (ShouldDeferPackageCreation == EShouldDeferPackageCreation::No)
+		{
+			// `MakePackage` will assert that we were deferring package creation (so that if a user calls it later,
+			// they actually asked to defer package creation before hand!). So we just switch it to yes to let
+			// `MakePackage` succeed, which will switch it back anyways.
+			ShouldDeferPackageCreation = EShouldDeferPackageCreation::Yes;
+			MakePackage();
+		}
+	}
+
+	void MakePackage()
+	{
+		check(ShouldDeferPackageCreation == EShouldDeferPackageCreation::Yes);
+
 		// Ensure that async loading is done.
 		FlushAsyncLoading();
 
 		// Create a package.
 		UPackage* const Package = CreatePackage(*PackageName);
-		
+
 		UObject* const Object = NewObject<UMyAutoRTFMTestObject>(Package, ObjectName, RF_Public | RF_Standalone);
 
 		// Need to mark it is loaded.
@@ -144,6 +169,8 @@ struct AUTORTFM_DISABLE FPackageScopedMaker final
 
 		// Then make sure the package is no longer loaded in our tables after.
 		REQUIRE(FindObject<UObject>(nullptr, *PackageName) == nullptr);
+
+		ShouldDeferPackageCreation = EShouldDeferPackageCreation::No;
 	}
 };
 
@@ -756,14 +783,12 @@ TEST_CASE("UPackage.AsyncLoading")
 		});
 	}
 
-	// TODO: these tests hit a random fail on CI that need investigating
-#if 0
 	SECTION("Packaged Load Fails First Time")
 	{
-		AutoRTFMTestUtils::FScopedRetry DisableRetry(AutoRTFM::ForTheRuntime::EAutoRTFMRetryTransactionState::NoRetry);
-
 		FString Name(FString::Printf(TEXT("/Game/%dAutoRTFMTestPackage%d"), FPlatformProcess::GetCurrentProcessId(), __LINE__));
 
+		FPackageScopedMaker Scope(Name, EShouldDeferPackageCreation::Yes);
+
 		{
 			bool bHit = false;
 
@@ -775,7 +800,7 @@ TEST_CASE("UPackage.AsyncLoading")
 			REQUIRE(bHit);
 		}
 
-		FPackageScopedMaker _(Name);
+		Scope.MakePackage();
 
 		{
 			bool bHit = false;
@@ -791,10 +816,10 @@ TEST_CASE("UPackage.AsyncLoading")
 
 	SECTION("Packaged Load Fails First Time, Second Try in Transaction")
 	{
-		AutoRTFMTestUtils::FScopedRetry DisableRetry(AutoRTFM::ForTheRuntime::EAutoRTFMRetryTransactionState::NoRetry);
-
 		FString Name(FString::Printf(TEXT("/Game/%dAutoRTFMTestPackage%d"), FPlatformProcess::GetCurrentProcessId(), __LINE__));
 
+		FPackageScopedMaker Scope(Name, EShouldDeferPackageCreation::Yes);
+
 		{
 			bool bHit = false;
 
@@ -806,7 +831,7 @@ TEST_CASE("UPackage.AsyncLoading")
 			REQUIRE(bHit);
 		}
 
-		FPackageScopedMaker _(Name);
+		Scope.MakePackage();
 
 		AutoRTFM::Transact([&]
 		{
@@ -825,6 +850,8 @@ TEST_CASE("UPackage.AsyncLoading")
 	{
 		FString Name(FString::Printf(TEXT("/Game/%dAutoRTFMTestPackage%d"), FPlatformProcess::GetCurrentProcessId(), __LINE__));
 
+		FPackageScopedMaker Scope(Name, EShouldDeferPackageCreation::Yes);
+
 		AutoRTFM::Transact([&]
 		{
 			bool bHit = false;
@@ -837,7 +864,7 @@ TEST_CASE("UPackage.AsyncLoading")
 			REQUIRE(bHit);
 		});
 
-		FPackageScopedMaker _(Name);
+		Scope.MakePackage();
 
 		AutoRTFM::Transact([&]
 		{
@@ -851,5 +878,4 @@ TEST_CASE("UPackage.AsyncLoading")
 			REQUIRE(bHit);
 		});
 	}
-#endif
 }
