[AutoRTFM] Avoid race condition in Stats test code.

I had Claude look for race conditions and it came up with a very plausible scenario. The fix seems reasonable and harmless.

The hypothetical race scenario:
1. DirectStatsCommand dispatches asynchronously -> bPrimaryEnable still false
2. WaitForStats barrier mechanism has 4-call latency -> first 4 calls return immediately (StatsSystemTypes.cpp:825-830: initial Barrier is default-constructed FTask, Wait() returns immediately)
3. While the 10-tick warmup should theoretically be sufficient, bPrimaryEnable is a non-atomic plain bool written cross-thread -> this is technically undefined behavior in C++, and the game thread may see stale values depending on platform memory model and compiler optimizations

With the blocking fix, Task.Wait() in DirectStatsCommand provides a proper happens-before relationship:
- HandleCommand completes before the call returns
- bPrimaryEnable = true with proper memory ordering
- NewFrame delegate is registered
- All subsequent SCOPE_CYCLE_COUNTER and AdvanceFrame calls see the correct state

#Synced-CL 51778913

diff --git a/Engine/Source/Programs/AutoRTFMTests/Private/StatsTests.cpp b/Engine/Source/Programs/AutoRTFMTests/Private/StatsTests.cpp
--- a/Engine/Source/Programs/AutoRTFMTests/Private/StatsTests.cpp
+++ b/Engine/Source/Programs/AutoRTFMTests/Private/StatsTests.cpp
@@ -75,7 +75,7 @@ TEST_CASE("StatsTests")
 		// Try and register FHUDGroupManager::NewFrame but this may not work
 		// Because we may be running with Stats already enabled (and that registration only happens
 		// when we first enable Stats).  Attempt it anyway.
-		UE::Stats::DirectStatsCommand(TEXT("stat AutoRTFM_Tests"));
+		UE::Stats::DirectStatsCommand(TEXT("stat AutoRTFM_Tests"), /*bBlockForCompletion=*/true);
 
 		// Let's fill the stats history
 		for (int Num = 0; Num < 10; ++Num)
