[AutoRTFM] Fix error handling issues reported by BugHawk.

- The Catch2 FAIL() method isn't supported outside of a test case
- Guarantee a non-zero return code when Result == 0 but FDebug::GetNumEnsureFailures > 0

#Synced-CL 51778913

diff --git a/Engine/Source/Programs/AutoRTFMTests/Private/AutoRTFMTests.cpp b/Engine/Source/Programs/AutoRTFMTests/Private/AutoRTFMTests.cpp
--- a/Engine/Source/Programs/AutoRTFMTests/Private/AutoRTFMTests.cpp
+++ b/Engine/Source/Programs/AutoRTFMTests/Private/AutoRTFMTests.cpp
@@ -152,7 +152,7 @@ AUTORTFM_DISABLE int main(int ArgC, const char* ArgV[])
 	// Install the AutoRTFM memory tracking allocator to test heap allocation correctness.
 	constexpr bool bRecordAllocationStackTraces = false;
 	AutoRTFM::Testing::FTrackingAllocator AutoRTFMTrackingAllocator([](const char* Error) {
-		FAIL("AutoRTFM tracking allocator failure: " << Error);
+		ensureMsgf(false, TEXT("AutoRTFM tracking allocator failure: %hs"), Error);
 	});
 	AutoRTFM::Testing::InterceptExternAPIs([&](autortfm_extern_api& API)
 	{
@@ -163,7 +163,7 @@ AUTORTFM_DISABLE int main(int ArgC, const char* ArgV[])
 	const int Result = RunTests(ArgC, ArgV);
 	if (Result != 0)
 	{
-		FAIL("AutoRTFM tests failed with code: " << Result);
+		ensureMsgf(false, TEXT("AutoRTFM tests failed with code: %d"), Result);
 		return Result;
 	}
 
@@ -172,7 +172,7 @@ AUTORTFM_DISABLE int main(int ArgC, const char* ArgV[])
 	{
 		// Change bRecordAllocationStackTraces to true to see leaked memory callstacks.
 		AutoRTFMTrackingAllocator.PrintAllocationCallstacks();
-		FAIL("AutoRTFM leaked " << NumAllocated << " bytes");
+		ensureMsgf(false, TEXT("AutoRTFM leaked %zu bytes"), NumAllocated);
 		return -1;
 	}
 #endif // AUTORTFM_ENABLE_TEST_UTILS
