[AutoRTFM] Re-enable AutoRTFM tests for Stats.

#Synced-CL 51757651

diff --git a/Engine/Source/Programs/AutoRTFMTests/Private/StatsTests.cpp b/Engine/Source/Programs/AutoRTFMTests/Private/StatsTests.cpp
--- a/Engine/Source/Programs/AutoRTFMTests/Private/StatsTests.cpp
+++ b/Engine/Source/Programs/AutoRTFMTests/Private/StatsTests.cpp
@@ -84,9 +84,6 @@ TEST_CASE("StatsTests")
 		}
 	}
 
-	// NOTE: (yiliang.siew) These have been disabled due to failing non-deterministically, and they were
-	// previously failing silently due to a different bug.
-#if 0
 	SECTION("Stats FScopedCycleCounter Outside Transaction")
 	{
 		SCOPE_CYCLE_COUNTER(STAT_StandardCycleCounter);
@@ -116,7 +113,6 @@ TEST_CASE("StatsTests")
 		NewValue = GetAutoRtfmStatValue();
 		CATCH_CHECK(NewValue > PreviousValue);
 	}
-#endif
 
 	SECTION("Stats Multi-FScopedCycleCounter Re-Entrant w/ Transaction")
 	{
