[AutoRTFM] Identify known Catch2 exit codes.

Catch2 exit codes are actually meaningful and useful, but numerically they are gibberish. TestFailureExitCode, for instance, is 42. Now we show `TestFailureExitCode` instead of `exit code: 42` when a CATCH_CHECK fails.

This only came up recently, because previously we would just crash on exit here due to the `FAIL()` from outside a test.

#Synced-CL 51778913

diff --git a/Engine/Source/Programs/AutoRTFMTests/Private/AutoRTFMTests.cpp b/Engine/Source/Programs/AutoRTFMTests/Private/AutoRTFMTests.cpp
--- a/Engine/Source/Programs/AutoRTFMTests/Private/AutoRTFMTests.cpp
+++ b/Engine/Source/Programs/AutoRTFMTests/Private/AutoRTFMTests.cpp
@@ -161,12 +161,31 @@ AUTORTFM_DISABLE int main(int ArgC, const char* ArgV[])
 #endif // AUTORTFM_ENABLE_TEST_UTILS
 
 	const int Result = RunTests(ArgC, ArgV);
-	if (Result != 0)
+
+#define CASE_CATCH_EXIT_CODE(X)														   \
+		case Catch::X:																   \
+			ensureMsgf(false, TEXT("AutoRTFM tests failed with ") TEXT(#X) TEXT(".")); \
+			return Result;
+
+	switch (Result)
 	{
-		ensureMsgf(false, TEXT("AutoRTFM tests failed with code: %d"), Result);
-		return Result;
+		CASE_CATCH_EXIT_CODE(UnspecifiedErrorExitCode)
+		CASE_CATCH_EXIT_CODE(NoTestsRunExitCode)
+		CASE_CATCH_EXIT_CODE(UnmatchedTestSpecExitCode)
+		CASE_CATCH_EXIT_CODE(AllTestsSkippedExitCode)
+		CASE_CATCH_EXIT_CODE(InvalidTestSpecExitCode)
+		CASE_CATCH_EXIT_CODE(TestFailureExitCode)
+
+		case 0:  // success
+			break;
+
+		default:
+			ensureMsgf(false, TEXT("AutoRTFM tests failed with code %d"), Result);
+			return Result;
 	}
 
+#undef CASE_CATCH_EXIT_CODE
+
 #if AUTORTFM_ENABLE_TEST_UTILS
 	if (size_t NumAllocated = AutoRTFMTrackingAllocator.TotalBytesAllocated(); NumAllocated != 0)
 	{
