[AutoRTFM] Add a hazard report for TMpscQueue.

This CL is equivalent to the changes made for TSpscQueue, and the hazard report should be sufficient for engineers to self-service when TMpscQueue is reached from within a transaction.

#Synced-CL 51757651

diff --git a/Engine/Source/Programs/AutoRTFMTests/Private/Abort.cpp b/Engine/Source/Programs/AutoRTFMTests/Private/Abort.cpp
--- a/Engine/Source/Programs/AutoRTFMTests/Private/Abort.cpp
+++ b/Engine/Source/Programs/AutoRTFMTests/Private/Abort.cpp
@@ -8,6 +8,7 @@
 #include "Async/RecursiveMutex.h"
 #include "Async/SharedMutex.h"
 #include "Catch2Includes.h"
+#include "Containers/MpscQueue.h"
 #include "Containers/SpscQueue.h"
 #include "Misc/StringOutputDevice.h"
 #include "Misc/OutputDeviceRedirector.h"
@@ -1586,19 +1587,19 @@ TEST_CASE("Abort.Language.ErrorHints.MallocOpenFree")
 }
 #endif
 
-TEST_CASE("Abort.Language.SpscQueue.SafeCtorDtor")
+TEMPLATE_TEST_CASE("Abort.Language.SpscMpscQueue.SafeCtorDtor", "", TSpscQueue<int>, TMpscQueue<int>)
 {
 	SECTION("Ctor")
 	{
 		AutoRTFM::Testing::Commit([&]
 		{
-			TSpscQueue<int> Queue;
+			TestType Queue;
 		});
 	}
 
 	SECTION("Dtor, Commit")
 	{
-		TUniquePtr<TSpscQueue<int>> Queue = MakeUnique<TSpscQueue<int>>();
+		TUniquePtr<TestType> Queue = MakeUnique<TestType>();
 		Queue->Enqueue(123);
 		Queue->Enqueue(456);
 
@@ -1610,7 +1611,7 @@ TEST_CASE("Abort.Language.SpscQueue.SafeCtorDtor")
 
 	SECTION("Dtor, Abort")
 	{
-		TUniquePtr<TSpscQueue<int>> Queue = MakeUnique<TSpscQueue<int>>();
+		TUniquePtr<TestType> Queue = MakeUnique<TestType>();
 		Queue->Enqueue(123);
 		Queue->Enqueue(456);
 
@@ -1620,16 +1621,16 @@ TEST_CASE("Abort.Language.SpscQueue.SafeCtorDtor")
 			AutoRTFM::AbortTransaction();
 		});
 
-		REQUIRE(Queue->Num() == 2);
+		REQUIRE(!Queue->IsEmpty());
 	}
 }
 
-TEST_CASE("Abort.Language.ErrorHints.SpscQueue")
+TEMPLATE_TEST_CASE("Abort.Language.ErrorHints.SpscMpscQueue", "", TSpscQueue<int>, TMpscQueue<int>)
 {
 	FScopedErrorHintTest ScopedTest;
     AutoRTFMTestUtils::FCaptureWarningContext WarningContext;
 
-	TSpscQueue<int> Queue;
+	TestType Queue;
 	AutoRTFM::ETransactionResult Result;
 
 	SECTION("Enqueue")
@@ -1674,5 +1675,18 @@ TEST_CASE("Abort.Language.ErrorHints.SpscQueue")
 	REQUIRE(AutoRTFM::ETransactionResult::AbortedByLanguage == Result);
 
 	// Did we print our message?
-	REQUIRE(WarningContext.HasWarningSubstring("An SPSC queue was accessed while inside an AutoRTFM transaction."));
+	const char* ExpectedMessage = "???";
+	if constexpr (std::is_same_v<TestType, TSpscQueue<int>>)
+	{
+		ExpectedMessage = "An SPSC queue was accessed while inside an AutoRTFM transaction.";
+	}
+	else if constexpr (std::is_same_v<TestType, TMpscQueue<int>>)
+	{
+		ExpectedMessage = "An MPSC queue was accessed while inside an AutoRTFM transaction.";
+	}
+	else
+	{
+		static_assert(!"unexpected TestType");
+	}
+	REQUIRE(WarningContext.HasWarningSubstring(ExpectedMessage));
 }
diff --git a/Engine/Source/Runtime/AutoRTFM/Private/ReportHazard.cpp b/Engine/Source/Runtime/AutoRTFM/Private/ReportHazard.cpp
--- a/Engine/Source/Runtime/AutoRTFM/Private/ReportHazard.cpp
+++ b/Engine/Source/Runtime/AutoRTFM/Private/ReportHazard.cpp
@@ -90,6 +90,16 @@ queues, but it does support transactionally-safe SPSC queues that can be used in
 
 Use the call stack below to find the TSpscQueue accessed immediately before this error occurred,
 and replace it with TTransactionallySafeSpscQueue.
+)";
+			break;
+
+		case EHazardType::MpscQueue:
+			Message = R"(
+An MPSC queue was accessed while inside an AutoRTFM transaction. AutoRTFM doesn't support native MPSC
+queues, but it does support transactionally-safe MPSC queues that can be used instead.
+
+Use the call stack below to find the TMpscQueue accessed immediately before this error occurred,
+and replace it with TTransactionallySafeMpscQueue.
 )";
 			break;
 
diff --git a/Engine/Source/Runtime/AutoRTFM/Public/AutoRTFM/ReportHazard.h b/Engine/Source/Runtime/AutoRTFM/Public/AutoRTFM/ReportHazard.h
--- a/Engine/Source/Runtime/AutoRTFM/Public/AutoRTFM/ReportHazard.h
+++ b/Engine/Source/Runtime/AutoRTFM/Public/AutoRTFM/ReportHazard.h
@@ -19,6 +19,7 @@ enum class EHazardType
 	SharedMutex,
 	OpenFree,
 	SpscQueue,
+	MpscQueue,
 };
 
 /**
diff --git a/Engine/Source/Runtime/Core/Public/Containers/MpscQueue.h b/Engine/Source/Runtime/Core/Public/Containers/MpscQueue.h
--- a/Engine/Source/Runtime/Core/Public/Containers/MpscQueue.h
+++ b/Engine/Source/Runtime/Core/Public/Containers/MpscQueue.h
@@ -3,6 +3,7 @@
 #pragma once
 
 #include "CoreTypes.h"
+#include "AutoRTFM/ReportHazard.h"
 #include "Templates/MemoryOps.h"
 #include "Templates/TypeCompatibleBytes.h"
 #include "Templates/UnrealTemplate.h"
@@ -21,11 +22,8 @@ public:
 
 	UE_NONCOPYABLE(TMpscQueue);
 
-	[[nodiscard]] TMpscQueue()
+	[[nodiscard]] TMpscQueue() : TMpscQueue(::new(AllocatorType::Malloc(sizeof(FNode), alignof(FNode))) FNode)
 	{
-		FNode* Sentinel = ::new(AllocatorType::Malloc(sizeof(FNode), alignof(FNode))) FNode;
-		Head.store(Sentinel, std::memory_order_relaxed);
-		Tail = Sentinel;
 	}
 
 	~TMpscQueue()
@@ -48,6 +46,8 @@ public:
 	template <typename... ArgTypes>
 	void Enqueue(ArgTypes&&... Args)
 	{
+		UE_AUTORTFM_REPORT_HAZARD_IF_CLOSED(MpscQueue);
+
 		FNode* New = ::new(AllocatorType::Malloc(sizeof(FNode), alignof(FNode))) FNode;
 		::new ((void*)&New->Value) ElementType(Forward<ArgTypes>(Args)...);
 
@@ -57,6 +57,8 @@ public:
 
 	TOptional<ElementType> Dequeue()
 	{
+		UE_AUTORTFM_REPORT_HAZARD_IF_CLOSED(MpscQueue);
+
 		FNode* Next = Tail->Next.load(std::memory_order_acquire);
 
 		if (Next == nullptr)
@@ -76,6 +78,8 @@ public:
 
 	bool Dequeue(ElementType& OutElem)
 	{
+		UE_AUTORTFM_REPORT_HAZARD_IF_CLOSED(MpscQueue);
+
 		TOptional<ElementType> LocalElement = Dequeue();
 		if (LocalElement.IsSet())
 		{
@@ -91,6 +95,8 @@ public:
 	// there's no overload with TOptional as it doesn't support references
 	[[nodiscard]] ElementType* Peek() const
 	{
+		UE_AUTORTFM_REPORT_HAZARD_IF_CLOSED(MpscQueue);
+
 		FNode* Next = Tail->Next.load(std::memory_order_acquire);
 
 		if (Next == nullptr)
@@ -103,6 +109,8 @@ public:
 
 	[[nodiscard]] bool IsEmpty() const
 	{
+		UE_AUTORTFM_REPORT_HAZARD_IF_CLOSED(MpscQueue);
+
 		return Tail->Next.load(std::memory_order_acquire) == nullptr;
 	}
 
@@ -113,6 +121,10 @@ private:
 		TTypeCompatibleBytes<ElementType> Value;
 	};
 
+	[[nodiscard]] explicit TMpscQueue(FNode* Sentinel) : Head(Sentinel), Tail(Sentinel)
+	{
+	}
+
 private:
 	std::atomic<FNode*> Head; // accessed only by producers
 	/*alignas(PLATFORM_CACHE_LINE_SIZE) */FNode* Tail; // accessed only by consumer, hence should be on a different cache line than `Head`
