[Backout] - CL51887653
#rnx
Original CL Desc
-----------------------------------------------------------------
[AutoRTFM] Fix error handling issues reported by BugHawk.

- The Catch2 FAIL() method isn't supported outside of a test case
- Guarantee a non-zero return code when Result == 0 but FDebug::GetNumEnsureFailures > 0

diff --git a/Engine/Source/Programs/AutoRTFMTests/Private/AutoRTFMTests.cpp b/Engine/Source/Programs/AutoRTFMTests/Private/AutoRTFMTests.cpp
--- a/Engine/Source/Programs/AutoRTFMTests/Private/AutoRTFMTests.cpp
+++ b/Engine/Source/Programs/AutoRTFMTests/Private/AutoRTFMTests.cpp
@@ -152,7 +152,7 @@ AUTORTFM_DISABLE int main(int ArgC, const char* ArgV[])
 	// Install the AutoRTFM memory tracking allocator to test heap allocation correctness.
 	constexpr bool bRecordAllocationStackTraces = false;
 	AutoRTFM::Testing::FTrackingAllocator AutoRTFMTrackingAllocator([](const char* Error) {
-		ensureMsgf(false, TEXT("AutoRTFM tracking allocator failure: %hs"), Error);
+		FAIL("AutoRTFM tracking allocator failure: " << Error);
 	});
 	AutoRTFM::Testing::InterceptExternAPIs([&](autortfm_extern_api& API)
 	{
@@ -160,41 +160,17 @@ AUTORTFM_DISABLE int main(int ArgC, const char* ArgV[])
 	});
 #endif // AUTORTFM_ENABLE_TEST_UTILS
 
-	const int Result = RunTests(ArgC, ArgV);
-
-#define CASE_CATCH_EXIT_CODE(X)														   \
-		case Catch::X:																   \
-			ensureMsgf(false, TEXT("AutoRTFM tests failed with ") TEXT(#X) TEXT(".")); \
-			return Result;
-
-	switch (Result)
-	{
-		CASE_CATCH_EXIT_CODE(UnspecifiedErrorExitCode)
-		CASE_CATCH_EXIT_CODE(NoTestsRunExitCode)
-		CASE_CATCH_EXIT_CODE(UnmatchedTestSpecExitCode)
-		CASE_CATCH_EXIT_CODE(AllTestsSkippedExitCode)
-		CASE_CATCH_EXIT_CODE(InvalidTestSpecExitCode)
-		CASE_CATCH_EXIT_CODE(TestFailureExitCode)
-
-		case 0:  // success
-			break;
-
-		default:
-			ensureMsgf(false, TEXT("AutoRTFM tests failed with code %d"), Result);
-			return Result;
-	}
-
-#undef CASE_CATCH_EXIT_CODE
+	int Result = RunTests(ArgC, ArgV);
 
 #if AUTORTFM_ENABLE_TEST_UTILS
 	if (size_t NumAllocated = AutoRTFMTrackingAllocator.TotalBytesAllocated(); NumAllocated != 0)
 	{
 		// Change bRecordAllocationStackTraces to true to see leaked memory callstacks.
 		AutoRTFMTrackingAllocator.PrintAllocationCallstacks();
-		ensureMsgf(false, TEXT("AutoRTFM leaked %zu bytes"), NumAllocated);
+		FAIL("AutoRTFM leaked " << NumAllocated << " bytes");
 		return -1;
 	}
 #endif // AUTORTFM_ENABLE_TEST_UTILS
 
-	return FDebug::GetNumEnsureFailures() == 0 ? 0 : -1;
+	return 0 == FDebug::GetNumEnsureFailures() ? 0 : -1;
 }
