Fix a bug whereby we were capturing transactionally allocated things by-value in a completion delegate and exploding during destruction.

Added a cvar `AutoRTFM.UseNewAbort` to let us revert back to the old behaviour to be safe.

diff --git a/Engine/Source/Programs/AutoRTFMTests/Private/UPackageTests.cpp b/Engine/Source/Programs/AutoRTFMTests/Private/UPackageTests.cpp
--- a/Engine/Source/Programs/AutoRTFMTests/Private/UPackageTests.cpp
+++ b/Engine/Source/Programs/AutoRTFMTests/Private/UPackageTests.cpp
@@ -646,6 +646,125 @@ TEST_CASE("UPackage.AsyncLoading")
 		REQUIRE(bHit);
 	}
 
+	SECTION("Transactionally allocated shared pointer captured by value in completion delegate")
+	{
+		FString Name(FString::Printf(TEXT("/Game/%dAutoRTFMTestPackage%d"), FPlatformProcess::GetCurrentProcessId(), __LINE__));
+		FPackageScopedMaker _(Name);
+
+		AutoRTFM::Transact([&]
+		{
+			TSharedRef<bool> TransactionallyCreatedPtr = MakeShared<bool>(false);
+
+			bool bHit = false;
+			FLoadPackageAsyncDelegate CompletionDelegate;
+			CompletionDelegate.BindLambda([&bHit, TransactionallyCreatedPtr](const FName&, UPackage*, EAsyncLoadingResult::Type Type) { bHit = true; REQUIRE(Type == EAsyncLoadingResult::Succeeded); });
+
+			FlushAsyncLoading(LoadPackageAsync(FPackagePath::FromPackageNameUnchecked(*Name), NAME_None, CompletionDelegate));
+
+			REQUIRE(bHit);
+		});
+	}
+
+	SECTION("Transactionally allocated shared pointer to shared from this captured by value in completion delegate")
+	{
+		FString Name(FString::Printf(TEXT("/Game/%dAutoRTFMTestPackage%d"), FPlatformProcess::GetCurrentProcessId(), __LINE__));
+		FPackageScopedMaker _(Name);
+
+		AutoRTFM::Transact([&]
+		{
+			struct FThing final : TSharedFromThis<FThing>
+			{
+				FThing(bool bInMember) : bMember(bInMember) {}
+				bool bMember = false;
+			};
+			TSharedRef<FThing> TransactionallyCreatedPtr = MakeShared<FThing>(false);
+
+			bool bHit = false;
+			FLoadPackageAsyncDelegate CompletionDelegate;
+			CompletionDelegate.BindLambda([&bHit, TransactionallyCreatedPtr](const FName&, UPackage*, EAsyncLoadingResult::Type Type) { bHit = true; REQUIRE(Type == EAsyncLoadingResult::Succeeded); });
+
+			FlushAsyncLoading(LoadPackageAsync(FPackagePath::FromPackageNameUnchecked(*Name), NAME_None, CompletionDelegate));
+
+			REQUIRE(bHit);
+		});
+	}
+
+	SECTION("Transactionally allocated weak pointer captured by value in completion delegate")
+	{
+		FString Name(FString::Printf(TEXT("/Game/%dAutoRTFMTestPackage%d"), FPlatformProcess::GetCurrentProcessId(), __LINE__));
+		FPackageScopedMaker _(Name);
+
+		AutoRTFM::Transact([&]
+		{
+			TSharedRef<bool> TransactionallyCreatedPtr = MakeShared<bool>(false);
+
+			TWeakPtr<bool> WeakPtr = TransactionallyCreatedPtr;
+
+			bool bHit = false;
+			FLoadPackageAsyncDelegate CompletionDelegate;
+			CompletionDelegate.BindLambda([&bHit, WeakPtr](const FName&, UPackage*, EAsyncLoadingResult::Type Type) { bHit = true; REQUIRE(Type == EAsyncLoadingResult::Succeeded); });
+
+			FlushAsyncLoading(LoadPackageAsync(FPackagePath::FromPackageNameUnchecked(*Name), NAME_None, CompletionDelegate));
+
+			REQUIRE(bHit);
+		});
+	}
+
+	SECTION("Transactionally allocated weak pointer to shared from this captured by value in completion delegate")
+	{
+		FString Name(FString::Printf(TEXT("/Game/%dAutoRTFMTestPackage%d"), FPlatformProcess::GetCurrentProcessId(), __LINE__));
+		FPackageScopedMaker _(Name);
+
+		AutoRTFM::Transact([&]
+		{
+			struct FThing final : TSharedFromThis<FThing>
+			{
+				FThing(bool bInMember) : bMember(bInMember) {}
+				bool bMember = false;
+			};
+			TSharedRef<FThing> TransactionallyCreatedPtr = MakeShared<FThing>(false);
+
+			TWeakPtr<FThing> WeakPtr = TransactionallyCreatedPtr;
+
+			bool bHit = false;
+			FLoadPackageAsyncDelegate CompletionDelegate;
+			CompletionDelegate.BindLambda([&bHit, WeakPtr](const FName&, UPackage*, EAsyncLoadingResult::Type Type) { bHit = true; REQUIRE(Type == EAsyncLoadingResult::Succeeded); });
+
+			FlushAsyncLoading(LoadPackageAsync(FPackagePath::FromPackageNameUnchecked(*Name), NAME_None, CompletionDelegate));
+
+			REQUIRE(bHit);
+		});
+	}
+
+	SECTION("Transactionally allocated shared pointer captured by value in completion delegate in a child")
+	{
+		FString Name1(FString::Printf(TEXT("/Game/%dAutoRTFMTestPackage%d"), FPlatformProcess::GetCurrentProcessId(), __LINE__));
+		FPackageScopedMaker _1(Name1);
+
+		FString Name2(FString::Printf(TEXT("/Game/%dAutoRTFMTestPackage%d"), FPlatformProcess::GetCurrentProcessId(), __LINE__));
+		FPackageScopedMaker _2(Name2);
+
+		AutoRTFM::Transact([&]
+		{
+			FLoadPackageAsyncDelegate CompletionDelegate;
+			LoadPackageAsync(FPackagePath::FromPackageNameUnchecked(*Name1), NAME_None, CompletionDelegate);
+
+			AutoRTFM::Transact([&]
+			{
+				TSharedRef<bool> TransactionallyCreatedPtr = MakeShared<bool>(false);
+
+				FLoadPackageAsyncDelegate CompletionDelegate;
+				CompletionDelegate.BindLambda([TransactionallyCreatedPtr](const FName&, UPackage*, EAsyncLoadingResult::Type Type) {});
+
+				LoadPackageAsync(FPackagePath::FromPackageNameUnchecked(*Name2), NAME_None, CompletionDelegate);
+
+				AutoRTFM::AbortTransaction();
+			});
+
+			AutoRTFM::AbortTransaction();
+		});
+	}
+
 	SECTION("First Time Load Package, Second Time Don't")
 	{
 		FString Name(FString::Printf(TEXT("/Game/%dAutoRTFMTestPackage%d"), FPlatformProcess::GetCurrentProcessId(), __LINE__));
diff --git a/Engine/Source/Runtime/CoreUObject/Private/Serialization/TransactionallySafeAsyncLoading.cpp b/Engine/Source/Runtime/CoreUObject/Private/Serialization/TransactionallySafeAsyncLoading.cpp
--- a/Engine/Source/Runtime/CoreUObject/Private/Serialization/TransactionallySafeAsyncLoading.cpp
+++ b/Engine/Source/Runtime/CoreUObject/Private/Serialization/TransactionallySafeAsyncLoading.cpp
@@ -12,6 +12,15 @@
 
 #if UE_AUTORTFM_ENABLED || defined(__INTELLISENSE__)
 
+namespace
+{
+static bool bAutoRTFMUseNewAbort = true;
+static FAutoConsoleVariableRef CVarAutoRTFMUseNewAbort(
+	TEXT("AutoRTFM.UseNewAbort"),
+	bAutoRTFMUseNewAbort,
+	TEXT("If true, use new aborting approach for completion delegate."));
+}
+
 // A transactionally safe async package loader that wraps an underlying actual
 // async package loader but allows it to be interacted with safely while also
 // inside a transaction.
@@ -179,6 +188,25 @@ public:
 					Payload.LoadFlags = InLoadFlags;
 					TransactionallyDeferredLoadPackages.Push(Payload);
 
+					const int32 Index = TransactionallyDeferredLoadPackages.Num() - 1;
+
+					if (LIKELY(bAutoRTFMUseNewAbort))
+					{
+						// We need to wipe out the completion delegate of our new deferred package load
+						// entry if we abort. This is because that completion delegate could have
+						// captured transactional local allocations that could be destroyed by the time
+						// we came to destruct all the items in `TransactionallyDeferredLoadPackages`
+						// when we empty it, causing use-after-frees.
+						const AutoRTFM::ETransactionStatus Status = AutoRTFM::Close([&]
+						{
+							AutoRTFM::OnAbort([this, Index]
+							{
+								TransactionallyDeferredLoadPackages[Index].CompletionDelegate.Unbind();
+							});
+						});
+						ensure(Status == AutoRTFM::ETransactionStatus::Executing);
+					}
+
 					// If we are the first, we need to register our handlers so that
 					// on-commit we issue the loads, and on-abort we purge the deferred
 					// array of them.
