[AutoRTFM] Apply more AutoRTFM attributes to the internal implementation to fix enabled -> disabled call edges.

#Synced-CL 52002210

diff --git a/Engine/Source/Programs/AutoRTFMTests/Private/HashMapTests.cpp b/Engine/Source/Programs/AutoRTFMTests/Private/HashMapTests.cpp
--- a/Engine/Source/Programs/AutoRTFMTests/Private/HashMapTests.cpp
+++ b/Engine/Source/Programs/AutoRTFMTests/Private/HashMapTests.cpp
@@ -12,6 +12,8 @@ using HashMapObjToObj = AutoRTFM::THashMap<AutoRTFMTestUtils::FObjectLifetimeHel
 
 TEMPLATE_TEST_CASE("HashMap", "", HashMapIntToInt, HashMapIntToObj, HashMapObjToInt, HashMapObjToObj)
 {
+	AutoRTFM::UnreachableIfClosed();
+
 	REQUIRE(AutoRTFMTestUtils::FObjectLifetimeHelper::ConstructorCalls == 0);
 	REQUIRE(AutoRTFMTestUtils::FObjectLifetimeHelper::DestructorCalls == 0);
 
@@ -19,7 +21,7 @@ TEMPLATE_TEST_CASE("HashMap", "", HashMapIntToInt, HashMapIntToObj, HashMapObjTo
 	using KeyType = typename TestType::Key;
 	using ValueType = typename TestType::Value;
 
-	auto Check = [&](MapType& Map, const std::unordered_map<KeyType, ValueType>& Expected)
+	auto Check = [&] AUTORTFM_DISABLE (MapType& Map, const std::unordered_map<KeyType, ValueType>& Expected)
 	{
 		// Check the reported count is as expected.
 		REQUIRE(Map.Num() == Expected.size());
diff --git a/Engine/Source/Runtime/AutoRTFM/Private/CallNest.h b/Engine/Source/Runtime/AutoRTFM/Private/CallNest.h
--- a/Engine/Source/Runtime/AutoRTFM/Private/CallNest.h
+++ b/Engine/Source/Runtime/AutoRTFM/Private/CallNest.h
@@ -12,7 +12,7 @@ namespace AutoRTFM
 
 class FContext;
 
-class FCallNest final
+class AUTORTFM_DISABLE FCallNest final
 {
 public:
 	FCallNest(FCallNest* Parent, TransactionID Transaction) : Parent(Parent), Transaction(Transaction) {}
diff --git a/Engine/Source/Runtime/AutoRTFM/Private/Context.h b/Engine/Source/Runtime/AutoRTFM/Private/Context.h
--- a/Engine/Source/Runtime/AutoRTFM/Private/Context.h
+++ b/Engine/Source/Runtime/AutoRTFM/Private/Context.h
@@ -116,17 +116,17 @@ public:
 
 	// Returns true if the caller's thread is the same as the thread associated
 	// with the transaction stack.
-	inline bool IsContextThread() const
+	AUTORTFM_ENABLE inline bool IsContextThread() const
 	{
 		return CurrentThreadId == FThreadID::GetCurrent();
 	}
 
-	inline EContextStatus GetStatus() const
+	AUTORTFM_ENABLE inline EContextStatus GetStatus() const
 	{
 		return IsContextThread() ? Status : EContextStatus::Idle;
 	}
 
-	inline bool IsTransactional() const
+	AUTORTFM_ENABLE inline bool IsTransactional() const
 	{
 		return GetStatus() == EContextStatus::OnTrack;
 	}
diff --git a/Engine/Source/Runtime/AutoRTFM/Private/HashMap.h b/Engine/Source/Runtime/AutoRTFM/Private/HashMap.h
--- a/Engine/Source/Runtime/AutoRTFM/Private/HashMap.h
+++ b/Engine/Source/Runtime/AutoRTFM/Private/HashMap.h
@@ -27,7 +27,7 @@ struct TKeyAndValue
 // improve performance and to avoid platform-specific variations in behavior.
 template <typename KeyType, typename ValueType, typename Hash = std::hash<KeyType>, typename Equal = std::equal_to<KeyType>,
 	Allocator Allocator = FAllocator>
-class THashMap
+class AUTORTFM_DISABLE THashMap
 {
 	using FInnerMap =
 		std::unordered_map<KeyType, ValueType, Hash, Equal, StlAllocator<std::pair<const KeyType, ValueType>, Allocator>>;
diff --git a/Engine/Source/Runtime/AutoRTFM/Private/LongJump.h b/Engine/Source/Runtime/AutoRTFM/Private/LongJump.h
--- a/Engine/Source/Runtime/AutoRTFM/Private/LongJump.h
+++ b/Engine/Source/Runtime/AutoRTFM/Private/LongJump.h
@@ -21,7 +21,7 @@ namespace AutoRTFM
 // setjmp/longjmp that doesn't do any unwinding (no C++ destructor calls, no messing with OS
 // signal states - just saving/restoring CPU state). Although it's the setjmp/longjmp everyone
 // knows and loves, it's exposed as a try/catch/throw API to make it less error-prone.
-class FLongJump
+class AUTORTFM_DISABLE FLongJump
 {
 public:
 	FLongJump()
diff --git a/Engine/Source/Runtime/AutoRTFM/Private/ThreadID.h b/Engine/Source/Runtime/AutoRTFM/Private/ThreadID.h
--- a/Engine/Source/Runtime/AutoRTFM/Private/ThreadID.h
+++ b/Engine/Source/Runtime/AutoRTFM/Private/ThreadID.h
@@ -19,7 +19,7 @@ namespace AutoRTFM
 
 // A unique identifier for a thread of execution.
 // Constructed as FThreadID::Invalid.
-struct AUTORTFM_DISABLE_NO_SANITIZE FThreadID
+struct AUTORTFM_OPEN FThreadID
 {
 	// An invalid thread identifier.
 	UE_AUTORTFM_API static const FThreadID Invalid;
