[AutoRTFM] Improve transactionally-safe recursive mutex performance by 20% with an object pool.

Rather than relying on SharedPtr and heap allocations for each recursive mutex, we now rely on a thread-safe TLockFreeFixedSizeAllocator to divvy out mutex state objects from a pool. Each state object takes exactly 16 bytes of memory, and there is zero extra per-object overhead. The allocator allocates space from the heap in 64K increments, for 4096 state objects at a time.

This adds some complexity to the code because we must manually refcount the state instead of relying on SharedPtr to do it. However, in practice, this is not a huge burden, and with some restructuring, the code is still easy enough to follow.

This reduces overall memory footprint and makes our benchmarks 12% to 29% faster.

A similar optimization should be possible for our other transactionally-safe mutex types (Mutex, SharedMutex, SharedRecursiveMutex) in followup CLs.

#Synced-CL 52137917

diff --git a/Engine/Source/Programs/AutoRTFMTests/Private/Catch2Includes.h b/Engine/Source/Programs/AutoRTFMTests/Private/Catch2Includes.h
--- a/Engine/Source/Programs/AutoRTFMTests/Private/Catch2Includes.h
+++ b/Engine/Source/Programs/AutoRTFMTests/Private/Catch2Includes.h
@@ -27,6 +27,7 @@ namespace Catch::Benchmark
     template <typename Fn, typename... Args>
     AUTORTFM_DISABLE inline auto invoke_deoptimized(Fn&& fn, Args&&... args) -> std::enable_if_t<std::is_same<void, decltype(fn(args...))>::value>;
 
+    struct AUTORTFM_DISABLE Benchmark;
     struct AUTORTFM_DISABLE Chronometer;
 
     namespace Detail { struct AUTORTFM_DISABLE BenchmarkFunction; }
diff --git a/Engine/Source/Programs/AutoRTFMTests/Private/MutexTests.cpp b/Engine/Source/Programs/AutoRTFMTests/Private/MutexTests.cpp
--- a/Engine/Source/Programs/AutoRTFMTests/Private/MutexTests.cpp
+++ b/Engine/Source/Programs/AutoRTFMTests/Private/MutexTests.cpp
@@ -8,6 +8,7 @@
 #include "Async/TransactionallySafeMutex.h"
 #include "Async/TransactionallySafeSharedMutex.h"
 #include "Async/TransactionallySafeSharedRecursiveMutex.h"
+#include "Async/TransactionallySafeMutexPool.h"
 
 #include <atomic>
 #include <condition_variable>
@@ -107,6 +108,8 @@ TEMPLATE_TEST_CASE("MutexTests", "",
 		}
 	};
 	
+	int NumLocksUsed = UE::Private::TransactionallySafeMutexPool::GetNumUsed();
+
 	SECTION("Construct Outside Transaction")
 	{
 		MutexType Mutex;
@@ -123,6 +126,8 @@ TEMPLATE_TEST_CASE("MutexTests", "",
 		});
 	}
 
+	REQUIRE(NumLocksUsed == UE::Private::TransactionallySafeMutexPool::GetNumUsed());
+
 	SECTION("Construct Inside Transaction")
 	{
 		AutoRTFM::Testing::Abort([&]
@@ -139,6 +144,8 @@ TEMPLATE_TEST_CASE("MutexTests", "",
 		});
 	}
 
+	REQUIRE(NumLocksUsed == UE::Private::TransactionallySafeMutexPool::GetNumUsed());
+
 	SECTION("IsLocked Outside Transaction")
 	{
 		MutexType Mutex;
@@ -150,6 +157,8 @@ TEMPLATE_TEST_CASE("MutexTests", "",
 		RequireNotLocked(Mutex);
 	}
 
+	REQUIRE(NumLocksUsed == UE::Private::TransactionallySafeMutexPool::GetNumUsed());
+
 	SECTION("IsLocked Inside Transaction")
 	{
 		MutexType Mutex;
@@ -167,6 +176,8 @@ TEMPLATE_TEST_CASE("MutexTests", "",
 		RequireNotLocked(Mutex);
 	}
 
+	REQUIRE(NumLocksUsed == UE::Private::TransactionallySafeMutexPool::GetNumUsed());
+
 	SECTION("Inside Transaction Used In Nested Transaction")
 	{
 		AutoRTFM::Testing::Abort([&]
@@ -214,6 +225,8 @@ TEMPLATE_TEST_CASE("MutexTests", "",
 		});
 	}
 
+	REQUIRE(NumLocksUsed == UE::Private::TransactionallySafeMutexPool::GetNumUsed());
+
 	SECTION("In Static Local Initializer")
 	{
 		struct AUTORTFM_ENABLE MyStruct final
@@ -221,9 +234,13 @@ TEMPLATE_TEST_CASE("MutexTests", "",
 			MutexType Mutex;
 		};
 
-		auto Lambda = [] AUTORTFM_ENABLE
+		auto Lambda = [&NumLocksUsed] AUTORTFM_ENABLE
 		{
 			static MyStruct Mine;
+
+			// The above line above allocates a new lock, so update our count.
+			NumLocksUsed = UE::Private::TransactionallySafeMutexPool::GetNumUsed();
+
 			UE::TScopeLock _(Mine.Mutex);
 			return 42;
 		};
@@ -244,6 +261,8 @@ TEMPLATE_TEST_CASE("MutexTests", "",
 		REQUIRE(42 == Lambda());
 	}
 
+	REQUIRE(NumLocksUsed == UE::Private::TransactionallySafeMutexPool::GetNumUsed());
+
 	SECTION("In Static Local Initializer Called From Open")
 	{
 		struct AUTORTFM_ENABLE MyStruct final
@@ -251,9 +270,13 @@ TEMPLATE_TEST_CASE("MutexTests", "",
 			MutexType Mutex;
 		};
 
-		auto Lambda = [] AUTORTFM_ENABLE
+		auto Lambda = [&NumLocksUsed] AUTORTFM_ENABLE
 		{
 			static MyStruct Mine;
+
+			// The above line above allocates a new lock, so update our count.
+			NumLocksUsed = UE::Private::TransactionallySafeMutexPool::GetNumUsed();
+
 			UE::TScopeLock _(Mine.Mutex);
 			return 42;
 		};
@@ -274,6 +297,8 @@ TEMPLATE_TEST_CASE("MutexTests", "",
 		REQUIRE(42 == Lambda());
 	}
 
+	REQUIRE(NumLocksUsed == UE::Private::TransactionallySafeMutexPool::GetNumUsed());
+
 	SECTION("Delete Heap Allocated Mutex")
 	{
 		SECTION("SingleThread") // Mutex owned and destructed by this thread
@@ -389,6 +414,8 @@ TEMPLATE_TEST_CASE("MutexTests", "",
 		}
 	}
 
+	REQUIRE(NumLocksUsed == UE::Private::TransactionallySafeMutexPool::GetNumUsed());
+
 	if constexpr (bSupportsUnbalancedLockingInTransaction)
 	{
 		SECTION("Lock Within, Unlock Outside")
@@ -406,6 +433,8 @@ TEMPLATE_TEST_CASE("MutexTests", "",
 		}
 	}
 
+	REQUIRE(NumLocksUsed == UE::Private::TransactionallySafeMutexPool::GetNumUsed());
+
 	if constexpr (bSupportsUnbalancedLockingInTransaction && bSupportsTryLockInTransaction)
 	{
 		SECTION("TryLock Within, Unlock Outside")
@@ -423,6 +452,8 @@ TEMPLATE_TEST_CASE("MutexTests", "",
 		}
 	}
 
+	REQUIRE(NumLocksUsed == UE::Private::TransactionallySafeMutexPool::GetNumUsed());
+
 	if constexpr (bSupportsUnbalancedLockingInTransaction)
 	{
 		SECTION("Lock Outside, Unlock Within")
@@ -440,6 +471,8 @@ TEMPLATE_TEST_CASE("MutexTests", "",
 		}
 	}
 
+	REQUIRE(NumLocksUsed == UE::Private::TransactionallySafeMutexPool::GetNumUsed());
+
 	if constexpr (bSupportsUnbalancedLockingInTransaction)
 	{
 		SECTION("Lock, Unlock, Lock")
@@ -474,6 +507,8 @@ TEMPLATE_TEST_CASE("MutexTests", "",
 		}
 	}
 
+	REQUIRE(NumLocksUsed == UE::Private::TransactionallySafeMutexPool::GetNumUsed());
+
 	if constexpr (bSupportsUnbalancedLockingInTransaction && bSupportsTryLockInTransaction)
 	{
 		SECTION("TryLock, Unlock, TryLock")
@@ -508,6 +543,8 @@ TEMPLATE_TEST_CASE("MutexTests", "",
 		}
 	}
 
+	REQUIRE(NumLocksUsed == UE::Private::TransactionallySafeMutexPool::GetNumUsed());
+
 	if constexpr (bSupportsUnbalancedLockingInTransaction)
 	{
 		SECTION("Unlock, Lock, Unlock")
@@ -543,6 +580,8 @@ TEMPLATE_TEST_CASE("MutexTests", "",
 		}
 	}
 
+	REQUIRE(NumLocksUsed == UE::Private::TransactionallySafeMutexPool::GetNumUsed());
+
 	if constexpr (bSupportsUnbalancedLockingInTransaction && bSupportsTryLockInTransaction)
 	{
 		SECTION("Unlock, TryLock, Unlock")
@@ -578,6 +617,8 @@ TEMPLATE_TEST_CASE("MutexTests", "",
 		}
 	}
 
+	REQUIRE(NumLocksUsed == UE::Private::TransactionallySafeMutexPool::GetNumUsed());
+
 	if constexpr (bSupportsUnbalancedLockingInTransaction)
 	{
 		SECTION("Lock Outside, Unlock & Lock Within")
@@ -596,6 +637,8 @@ TEMPLATE_TEST_CASE("MutexTests", "",
 		}
 	}
 
+	REQUIRE(NumLocksUsed == UE::Private::TransactionallySafeMutexPool::GetNumUsed());
+
 	if constexpr (bSupportsUnbalancedLockingInTransaction)
 	{
 		SECTION("Commit(Lock, Commit(Unlock))")
@@ -618,6 +661,8 @@ TEMPLATE_TEST_CASE("MutexTests", "",
 		}
 	}
 
+	REQUIRE(NumLocksUsed == UE::Private::TransactionallySafeMutexPool::GetNumUsed());
+
 	if constexpr (bSupportsUnbalancedLockingInTransaction)
 	{
 		SECTION("Abort(Lock, Commit(Unlock, Lock))")
@@ -643,6 +688,8 @@ TEMPLATE_TEST_CASE("MutexTests", "",
 		}
 	}
 
+	REQUIRE(NumLocksUsed == UE::Private::TransactionallySafeMutexPool::GetNumUsed());
+
 	if constexpr (bSupportsUnbalancedLockingInTransaction)
 	{
 		SECTION("Contended Lock")
@@ -713,6 +760,8 @@ TEMPLATE_TEST_CASE("MutexTests", "",
 		}
 	}
 
+	REQUIRE(NumLocksUsed == UE::Private::TransactionallySafeMutexPool::GetNumUsed());
+
 	SECTION("In On-Commit")
 	{
 		if constexpr (bSupportsUnbalancedLockingInTransaction)
@@ -873,6 +922,8 @@ TEMPLATE_TEST_CASE("MutexTests", "",
 		}
 	}
 
+	REQUIRE(NumLocksUsed == UE::Private::TransactionallySafeMutexPool::GetNumUsed());
+
 	SECTION("In On-Abort")
 	{
 		SECTION("Transact(Check-On-Abort, Lock, Check-On-Abort, Abort)")
@@ -1007,6 +1058,8 @@ TEMPLATE_TEST_CASE("MutexTests", "",
 		}
 	}
 
+	REQUIRE(NumLocksUsed == UE::Private::TransactionallySafeMutexPool::GetNumUsed());
+
 	SECTION("Locked Mutex In Destructed Object")
 	{
 		struct AUTORTFM_ENABLE MyStruct final
@@ -1015,6 +1068,7 @@ TEMPLATE_TEST_CASE("MutexTests", "",
 
 			~MyStruct()
 			{
+				// This zeroes out the mutex's state pointer, causing it to leak; real code shouldn't do this!
 				memset(this, 0, sizeof(MyStruct));
 			}
 		};
@@ -1027,8 +1081,13 @@ TEMPLATE_TEST_CASE("MutexTests", "",
 			Mine.reset();
 			REQUIRE(!Mine);
 		});
+
+		// The lock inside MyStruct is now permanently leaked. We account for this pool size increase here.
+		NumLocksUsed = UE::Private::TransactionallySafeMutexPool::GetNumUsed();
 	}
 
+	REQUIRE(NumLocksUsed == UE::Private::TransactionallySafeMutexPool::GetNumUsed());
+
 	if constexpr (bSupportsMixedOpenCloseLocking)
 	{
 		SECTION("Closed Lock Then Open Unlock")
@@ -1048,6 +1107,8 @@ TEMPLATE_TEST_CASE("MutexTests", "",
 		}
 	}
 
+	REQUIRE(NumLocksUsed == UE::Private::TransactionallySafeMutexPool::GetNumUsed());
+
 	if constexpr (bSupportsMixedOpenCloseLocking && bSupportsTryLockInTransaction)
 	{
 		SECTION("Closed TryLock Then Open Unlock")
@@ -1067,6 +1128,8 @@ TEMPLATE_TEST_CASE("MutexTests", "",
 		}
 	}
 
+	REQUIRE(NumLocksUsed == UE::Private::TransactionallySafeMutexPool::GetNumUsed());
+
 	if constexpr (bSupportsMixedOpenCloseLocking)
 	{
 		SECTION("Open Lock Then Closed Unlock")
@@ -1086,6 +1149,8 @@ TEMPLATE_TEST_CASE("MutexTests", "",
 		}
 	}
 
+	REQUIRE(NumLocksUsed == UE::Private::TransactionallySafeMutexPool::GetNumUsed());
+
 	if constexpr (bSupportsMixedOpenCloseLocking && bSupportsTryLockInTransaction)
 	{
 		SECTION("Open TryLock Then Closed Unlock")
@@ -1105,6 +1170,8 @@ TEMPLATE_TEST_CASE("MutexTests", "",
 		}
 	}
 
+	REQUIRE(NumLocksUsed == UE::Private::TransactionallySafeMutexPool::GetNumUsed());
+
 	if constexpr (bIsRecursive)
 	{
 		SECTION("Transact(Lock(Lock()))")
@@ -1127,6 +1194,8 @@ TEMPLATE_TEST_CASE("MutexTests", "",
 		}
 	}
 
+	REQUIRE(NumLocksUsed == UE::Private::TransactionallySafeMutexPool::GetNumUsed());
+
 	if constexpr (bIsRecursive)
 	{
 		SECTION("Lock(Lock())")
@@ -1145,6 +1214,8 @@ TEMPLATE_TEST_CASE("MutexTests", "",
 		}
 	}
 
+	REQUIRE(NumLocksUsed == UE::Private::TransactionallySafeMutexPool::GetNumUsed());
+
 	if constexpr (bIsRecursive)
 	{
 		SECTION("Transact(Lock(Lock(Abort)))")
@@ -1167,6 +1238,8 @@ TEMPLATE_TEST_CASE("MutexTests", "",
 		}
 	}
 
+	REQUIRE(NumLocksUsed == UE::Private::TransactionallySafeMutexPool::GetNumUsed());
+
 	if constexpr (bIsRecursive)
 	{
 		SECTION("Transact(Lock(Lock(), Abort))")
@@ -1189,6 +1262,8 @@ TEMPLATE_TEST_CASE("MutexTests", "",
 		}
 	}
 
+	REQUIRE(NumLocksUsed == UE::Private::TransactionallySafeMutexPool::GetNumUsed());
+
 	if constexpr (bIsRecursive)
 	{
 		SECTION("Transact(Lock(Transact(Lock(), Abort)))")
@@ -1211,6 +1286,8 @@ TEMPLATE_TEST_CASE("MutexTests", "",
 		}
 	}
 
+	REQUIRE(NumLocksUsed == UE::Private::TransactionallySafeMutexPool::GetNumUsed());
+
 	if constexpr (bIsShared && bIsRecursive)
 	{
 		SECTION("SharedLock(SharedLock())")
@@ -1235,6 +1312,8 @@ TEMPLATE_TEST_CASE("MutexTests", "",
 		}
 	}
 
+	REQUIRE(NumLocksUsed == UE::Private::TransactionallySafeMutexPool::GetNumUsed());
+
 	if constexpr (bIsShared && bIsRecursive)
 	{
 		SECTION("Transact(SharedLock(SharedLock()))")
@@ -1261,6 +1340,8 @@ TEMPLATE_TEST_CASE("MutexTests", "",
 		}
 	}
 
+	REQUIRE(NumLocksUsed == UE::Private::TransactionallySafeMutexPool::GetNumUsed());
+
 	if constexpr (bIsShared && bIsRecursive)
 	{
 		SECTION("Transact(SharedLock(SharedLock(Abort)))")
@@ -1286,6 +1367,8 @@ TEMPLATE_TEST_CASE("MutexTests", "",
 		}
 	}
 
+	REQUIRE(NumLocksUsed == UE::Private::TransactionallySafeMutexPool::GetNumUsed());
+
 	if constexpr (bIsShared && bIsRecursive)
 	{
 		SECTION("Transact(SharedLock(SharedLock(), Abort))")
@@ -1311,6 +1394,8 @@ TEMPLATE_TEST_CASE("MutexTests", "",
 		}
 	}
 
+	REQUIRE(NumLocksUsed == UE::Private::TransactionallySafeMutexPool::GetNumUsed());
+
 	if constexpr (bIsShared && bIsRecursive)
 	{
 		SECTION("Transact(SharedLock(Transact(SharedLock(), Abort)))")
@@ -1336,4 +1421,125 @@ TEMPLATE_TEST_CASE("MutexTests", "",
 			RequireNotSharedLocked(Mutex);
 		}
 	}
+
+	REQUIRE(NumLocksUsed == UE::Private::TransactionallySafeMutexPool::GetNumUsed());
+
+	SECTION("Benchmarks")
+	{
+		constexpr int Iterations = 10000;
+
+		BENCHMARK("Commit: Loop(Transact(Ctor, Lock, Unlock, Dtor))")
+		{
+			for (int Index = 0; Index < Iterations; ++Index)
+			{
+				AutoRTFM::Testing::Commit([&]
+				{
+					MutexType Mutex;
+					Mutex.Lock();
+					Mutex.Unlock();
+				});
+			}
+		};
+
+		BENCHMARK("Commit: Transact(Loop(Ctor, Lock, Unlock, Dtor))")
+		{
+			AutoRTFM::Testing::Commit([&]
+			{
+				for (int Index = 0; Index < Iterations; ++Index)
+				{
+					MutexType Mutex;
+					Mutex.Lock();
+					Mutex.Unlock();
+				}
+			});
+		};
+
+		BENCHMARK("Commit: Transact(Loop(Lock, Unlock))")
+		{
+			MutexType Mutex;
+			AutoRTFM::Testing::Commit([&]
+			{
+				for (int Index = 0; Index < Iterations; ++Index)
+				{
+					Mutex.Lock();
+					Mutex.Unlock();
+				}
+			});
+		};
+
+		BENCHMARK("Commit: Transact(Loop(Lock), Loop(Unlock))")
+		{
+			MutexType Mutex[Iterations];
+			AutoRTFM::Testing::Commit([&]
+			{
+				for (int Index = 0; Index < Iterations; ++Index)
+				{
+					Mutex[Index].Lock();
+				}
+				for (int Index = 0; Index < Iterations; ++Index)
+				{
+					Mutex[Index].Unlock();
+				}
+			});
+		};
+
+		BENCHMARK("Abort: Loop(Transact(Ctor, Lock, Unlock, Dtor))")
+		{
+			for (int Index = 0; Index < Iterations; ++Index)
+			{
+				AutoRTFM::Testing::Abort([&]
+				{
+					MutexType Mutex;
+					Mutex.Lock();
+					Mutex.Unlock();
+					AutoRTFM::AbortTransaction();
+				});
+			}
+		};
+
+		BENCHMARK("Abort: Transact(Loop(Ctor, Lock, Unlock, Dtor))")
+		{
+			AutoRTFM::Testing::Abort([&]
+			{
+				for (int Index = 0; Index < Iterations; ++Index)
+				{
+					MutexType Mutex;
+					Mutex.Lock();
+					Mutex.Unlock();
+				}
+				AutoRTFM::AbortTransaction();
+			});
+		};
+
+		BENCHMARK("Abort: Transact(Loop(Lock, Unlock))")
+		{
+			MutexType Mutex;
+			AutoRTFM::Testing::Abort([&]
+			{
+				for (int Index = 0; Index < Iterations; ++Index)
+				{
+					Mutex.Lock();
+					Mutex.Unlock();
+				}
+				AutoRTFM::AbortTransaction();
+			});
+		};
+
+		BENCHMARK("Abort: Transact(Loop(Lock), Loop(Unlock))")
+		{
+			MutexType Mutex[Iterations];
+			AutoRTFM::Testing::Abort([&]
+			{
+				for (int Index = 0; Index < Iterations; ++Index)
+				{
+					Mutex[Index].Lock();
+				}
+				for (int Index = 0; Index < Iterations; ++Index)
+				{
+					Mutex[Index].Unlock();
+				}
+				AutoRTFM::AbortTransaction();
+			});
+		};
+	}
 }
diff --git a/Engine/Source/Runtime/Core/Private/Async/TransactionallySafeMutexPool.cpp b/Engine/Source/Runtime/Core/Private/Async/TransactionallySafeMutexPool.cpp
new file mode 100644
--- /dev/null
+++ b/Engine/Source/Runtime/Core/Private/Async/TransactionallySafeMutexPool.cpp
@@ -0,0 +1,50 @@
+// Copyright Epic Games, Inc. All Rights Reserved.
+
+#include "Async/TransactionallySafeMutexPool.h"
+#include "AutoRTFM.h"
+#include "Containers/LockFreeFixedSizeAllocator.h"
+#include "HAL/ThreadSafeCounter.h"
+
+#if !UE_BUILD_SHIPPING
+using FAllocatorType = TLockFreeFixedSizeAllocator_TLSCache<UE::Private::TransactionallySafeMutexPool::PooledObjectSize, /*TPaddingForCacheContention*/0, FThreadSafeCounter>;
+#else
+using FAllocatorType = TLockFreeFixedSizeAllocator_TLSCache<UE::Private::TransactionallySafeMutexPool::PooledObjectSize, /*TPaddingForCacheContention*/0>;
+#endif
+
+alignas(FAllocatorType) static uint8 GAllocatorStorage[sizeof(FAllocatorType)];
+
+static FAllocatorType& SingletonMutexAllocator()
+{
+	static FAllocatorType* Allocator = new (GAllocatorStorage) FAllocatorType;
+	return *Allocator;
+}
+
+void* UE::Private::TransactionallySafeMutexPool::Allocate()
+{
+	void* Ptr = nullptr;
+	UE_AUTORTFM_OPEN
+	{
+		Ptr = SingletonMutexAllocator().Allocate();
+	};
+	UE_AUTORTFM_ONABORT(Ptr)
+	{
+		SingletonMutexAllocator().Free(Ptr);
+	};
+	return AutoRTFM::DidAllocate(Ptr, PooledObjectSize);
+}
+
+void UE::Private::TransactionallySafeMutexPool::Free(void* Ptr)
+{
+	UE_AUTORTFM_ONCOMMIT(Ptr)
+	{
+		SingletonMutexAllocator().Free(Ptr);
+		AutoRTFM::DidFree(Ptr);
+	};
+}
+
+#if !UE_BUILD_SHIPPING
+int UE::Private::TransactionallySafeMutexPool::GetNumUsed()
+{
+	return SingletonMutexAllocator().GetNumUsed().GetValue();
+}
+#endif
\ No newline at end of file
diff --git a/Engine/Source/Runtime/Core/Public/Async/TransactionallySafeMutexPool.h b/Engine/Source/Runtime/Core/Public/Async/TransactionallySafeMutexPool.h
new file mode 100644
--- /dev/null
+++ b/Engine/Source/Runtime/Core/Public/Async/TransactionallySafeMutexPool.h
@@ -0,0 +1,47 @@
+// Copyright Epic Games, Inc. All Rights Reserved.
+
+#pragma once
+#include "AutoRTFM/Defines.h"
+#include "HAL/Platform.h"
+
+#define UE_API CORE_API
+
+namespace UE::Private::TransactionallySafeMutexPool
+{
+
+// The size of an object in the allocator pool is 16 bytes. This value is chosen to match
+// `sizeof(FTransactionallySafeRecursiveMutex::FState)`. The other transactionally-safe
+// mutexes have state objects which also fit into 16 bytes or less.
+constexpr size_t PooledObjectSize = 16;
+
+/** Takes `PooledObjectSize` raw bytes from our mutex pool. */
+UE_API void* Allocate();
+
+/** Constructs a new state object, taking it from the pool. */
+template <typename T>
+	requires (sizeof(T) <= PooledObjectSize)
+UE_REWRITE T* New()
+{
+	return new (UE::Private::TransactionallySafeMutexPool::Allocate()) T;
+}
+
+/** Returns `PooledObjectSize` raw bytes to the mutex pool. */
+UE_API void Free(void* Ptr);
+
+/** Destroys one state object, returning it to the pool. */
+template <typename T>
+	requires (sizeof(T) <= PooledObjectSize)
+UE_REWRITE void Delete(T* Ptr)
+{
+	Ptr->~T();
+	UE::Private::TransactionallySafeMutexPool::Free(Ptr);
+}
+
+#if !UE_BUILD_SHIPPING
+/** Returns the number of objects currently allocated from the pool. Intended for testing. */
+UE_API int GetNumUsed();
+#endif
+
+} // namespace UE::Private::TransactionallySafeMutexPool
+
+#undef UE_API
diff --git a/Engine/Source/Runtime/Core/Public/Async/TransactionallySafeRecursiveMutex.h b/Engine/Source/Runtime/Core/Public/Async/TransactionallySafeRecursiveMutex.h
--- a/Engine/Source/Runtime/Core/Public/Async/TransactionallySafeRecursiveMutex.h
+++ b/Engine/Source/Runtime/Core/Public/Async/TransactionallySafeRecursiveMutex.h
@@ -6,6 +6,7 @@
 #include "Async/RecursiveMutex.h"
 
 #if UE_AUTORTFM
+#include "Async/TransactionallySafeMutexPool.h"
 #include "AutoRTFM/OpenWrapper.h"
 #include "Templates/SharedPointer.h"
 #endif
@@ -34,83 +35,95 @@ namespace UE
 // modifications we should make.
 struct FTransactionallySafeRecursiveMutexDefinition final
 {
-	FTransactionallySafeRecursiveMutexDefinition() : State(MakeShared<FState>())
+	FTransactionallySafeRecursiveMutexDefinition() : State(UE::Private::TransactionallySafeMutexPool::New<FState>())
 	{
 	}
 
+	~FTransactionallySafeRecursiveMutexDefinition()
+	{
+		// Normally, State should never be null, but test "Locked Mutex In Destructed Object" expects
+		// that we can memzero a mutex safely. This doesn't crash, but leaks the mutex's State pointer.
+		if (State)
+		{
+			State->ReleaseOnCommit();
+			State = nullptr;
+		}
+	}
+
 	void Lock()
 	{
-		if (!AutoRTFM::IsIdle())
+		if (AutoRTFM::IsIdle())
+		{
+			State->Mutex.Lock();
+			ensure(0 == State->TransactionalLockCount);
+		}
+		else if (AutoRTFM::IsClosed())
 		{
 			UE_AUTORTFM_OPEN
 			{
-				// The transactional system which can increment TransactionalLockCount
-				// is always single-threaded, thus this is safe to check without atomicity.
-				if (0 == State->TransactionalLockCount)
-				{
-					State->Mutex.Lock();
-				}
-
-				State->TransactionalLockCount += 1;
+				// We explicitly extend the lifetime of `State` here, in case `this` lives on the
+				// stack and ends up being overwritten before the on-abort is hit.
+				State->Acquire();
+				State->LockTransactionalLock();
 			};
 
-			// We explicitly copy the state here for the case that `this` was stack
-			// allocated and has already died before the on-abort is hit.
-			AutoRTFM::OnAbort([State = AutoRTFM::TOpenWrapper{this->State}]
+			AutoRTFM::OnAbort([State = this->State]
 			{
-				ensure(0 != State.Object->TransactionalLockCount);
-
-				State.Object->TransactionalLockCount -= 1;
-
-				if (0 == State.Object->TransactionalLockCount)
-				{
-					State.Object->Mutex.Unlock();
-				}
+				// Undo the lock if the transaction is aborted.
+				State->UnlockTransactionalLock();
 			});
+
+			State->ReleaseOnComplete();
 		}
 		else
 		{
-			State->Mutex.Lock();
-			ensure(0 == State->TransactionalLockCount);
+			// We're inside a transaction, but not in the closed.
+			State->LockTransactionalLock();
 		}
 	}
 
 	void Unlock()
 	{
-		if (!AutoRTFM::IsIdle())
+		if (AutoRTFM::IsIdle())
 		{
-			// We explicitly copy the state here for the case that `this` was stack
-			// allocated and has already died before the on-commit is hit.
-			AutoRTFM::OnCommit([State = AutoRTFM::TOpenWrapper{this->State}]
+			ensure(0 == State->TransactionalLockCount);
+			State->Mutex.Unlock();
+		}
+		else if (AutoRTFM::IsClosed())
+		{
+			UE_AUTORTFM_OPEN
 			{
-				ensure(0 != State.Object->TransactionalLockCount);
-
-				State.Object->TransactionalLockCount -= 1;
+				// We explicitly extend the lifetime of `State` here, in case `this` lives on the
+				// stack and ends up being overwritten before the on-commit is hit.
+				State->Acquire();
+			};
 
-				if (0 == State.Object->TransactionalLockCount)
-				{
-					State.Object->Mutex.Unlock();
-				}
+			AutoRTFM::OnCommit([State = this->State]
+			{
+				// We don't want to actually unlock the transactional mutex until the transaction reaches commit.
+				State->UnlockTransactionalLock();
 			});
+
+			State->ReleaseOnComplete();
 		}
 		else
 		{
-			ensure(0 == State->TransactionalLockCount);
-			State->Mutex.Unlock();
+			// We're inside a transaction, but not in the closed.
+			State->UnlockTransactionalLock();
 		}
 	}
 
 	[[nodiscard]] bool TryLock()
 	{
-		if (!AutoRTFM::IsIdle())
+		if (AutoRTFM::IsIdle())
 		{
-			// TODO: implement transactional TryLock. #jira SOL-8303
-			return false;
+			ensure(0 == State->TransactionalLockCount);
+			return State->Mutex.TryLock();
 		}
 		else
 		{
-			ensure(0 == State->TransactionalLockCount);
-			return State->Mutex.TryLock();
+			// TODO: implement transactional TryLock. #jira SOL-8303
+			return false;
 		}
 	}
 
@@ -119,18 +132,85 @@ private:
 
 	struct FState final
 	{
-		UE::FRecursiveMutex Mutex;
-		uint32 TransactionalLockCount = 0;
-
 		FState() = default;
-		
 		~FState()
 		{
 			ensure(0 == TransactionalLockCount);
+			checkSlow(0 == RefCount);
+		}
+
+		UE::FRecursiveMutex Mutex;
+		uint32 TransactionalLockCount = 0;
+		uint32 RefCount = 1;
+
+		AUTORTFM_DISABLE void LockTransactionalLock()
+		{
+			// The transactional system which can increment TransactionalLockCount
+			// is always single-threaded, thus this is safe to check without atomicity.
+			if (0 == TransactionalLockCount)
+			{
+				Mutex.Lock();
+			}
+
+			TransactionalLockCount += 1;
+		}
+
+		AUTORTFM_DISABLE void UnlockTransactionalLock()
+		{
+			ensure(TransactionalLockCount > 0);
+
+			TransactionalLockCount -= 1;
+			if (TransactionalLockCount == 0)
+			{
+				Mutex.Unlock();
+			}
+		}
+
+		AUTORTFM_DISABLE void Acquire()
+		{
+			RefCount += 1;
+		}
+
+		void ReleaseOnComplete()
+		{
+			checkSlow(RefCount > 0);
+
+			if (AutoRTFM::IsClosed())
+			{
+				AutoRTFM::OnComplete([this]
+				{
+					ReleaseImmediately();
+				});
+			}
+			else
+			{
+				ReleaseImmediately();
+			}
+		}
+
+		void ReleaseOnCommit()
+		{
+			checkSlow(RefCount > 0);
+
+			AutoRTFM::OnCommit([this]
+			{
+				ReleaseImmediately();
+			});
+		}
+
+		AUTORTFM_DISABLE void ReleaseImmediately()
+		{
+			checkSlow(RefCount > 0);
+
+			RefCount -= 1;
+			if (RefCount == 0)
+			{
+				UE::Private::TransactionallySafeMutexPool::Delete(this);
+			}
 		}
 	};
 
-	TSharedPtr<FState> State;
+	FState* State = nullptr;
 };
 
 using FTransactionallySafeRecursiveMutex = ::UE::FTransactionallySafeRecursiveMutexDefinition;
