LLM: AutoRTFMTests: Allow LLM code to be compiled into the project, and allow LLM to be enabled on commandline with -llm.
#rnx

#Synced-CL 52568437

diff --git a/Engine/Source/Programs/AutoRTFMTests/AutoRTFMTests.Target.cs b/Engine/Source/Programs/AutoRTFMTests/AutoRTFMTests.Target.cs
--- a/Engine/Source/Programs/AutoRTFMTests/AutoRTFMTests.Target.cs
+++ b/Engine/Source/Programs/AutoRTFMTests/AutoRTFMTests.Target.cs
@@ -94,6 +94,9 @@ public class AutoRTFMTestsTarget : TargetRules
 		// Need to force the heartbeat on so we can test with it!
 		GlobalDefinitions.Add("USE_HITCH_DETECTION=1");
 
+		// Support unittests for LowLevelMemTracker
+		GlobalDefinitions.Add("LLM_ALLOWED_IN_APPLICATION_TYPE=1");
+
 		// For testing insights, this sets `UE_TRACE_ENABLED`.
 		bEnableTrace = true;
 	}
diff --git a/Engine/Source/Programs/AutoRTFMTests/Private/AutoRTFMTests.cpp b/Engine/Source/Programs/AutoRTFMTests/Private/AutoRTFMTests.cpp
--- a/Engine/Source/Programs/AutoRTFMTests/Private/AutoRTFMTests.cpp
+++ b/Engine/Source/Programs/AutoRTFMTests/Private/AutoRTFMTests.cpp
@@ -6,7 +6,9 @@
 #include "CoreMinimal.h"
 #include "Misc/EventPool.h"
 #include "Misc/LazySingleton.h"
+#include "Misc/CString.h"
 #include "Misc/PackageName.h"
+#include "Misc/StringBuilder.h"
 #include "RequiredProgramMainCPPInclude.h" // required for ue programs
 #include "UObject/GCObject.h"
 #include "Serialization/PackageStore.h"
@@ -41,6 +43,26 @@ int RunTests(int ArgC, const char* ArgV[])
 #endif
 	);
 
+	TStringBuilder<256> CommandLine(InPlace,
+		TEXT("-Multiprocess -csvNoProcessingThread -LogCmds=\"LogCsvProfiler off, LogStreaming off, LogUObjectGlobals off, LogPackageName off, LogAutoRTFM warning\" -AsyncLoadingThread"));
+#if FORCE_USE_STATS
+	CommandLine << TEXT(" -LoadTimeStatsForCommandlet");
+#endif
+	// Allow callers to pass "... -- <UnrealArgs>" on the commandline
+	int OriginalArgC = ArgC;
+	for (int N = 0; N < OriginalArgC; ++N)
+	{
+		if (0 == FCStringAnsi::Strcmp(ArgV[N], "--"))
+		{
+			ArgC = N;
+			for (int CopiedIndex = N + 1; CopiedIndex < OriginalArgC; ++CopiedIndex)
+			{
+				CommandLine << TEXT(" ") << ArgV[CopiedIndex];
+			}
+			break;
+		}
+	}
+	
 	{
 		const int Result = Session.applyCommandLine(ArgC, ArgV);
 
@@ -59,10 +81,6 @@ int RunTests(int ArgC, const char* ArgV[])
 
 	AUTORTFM_SANITIZER_ONLY(AutoRTFM::Sanitizer::SetRecordClosedWriteCallstacks(bSanitizerRecordWriteCallstacks));
 
-	FString CommandLine = TEXT("-Multiprocess -csvNoProcessingThread -LogCmds=\"LogCsvProfiler off, LogStreaming off, LogUObjectGlobals off, LogPackageName off, LogAutoRTFM warning\" -AsyncLoadingThread");
-#if FORCE_USE_STATS
-	CommandLine += TEXT(" -LoadTimeStatsForCommandlet");
-#endif
 	GEngineLoop.PreInit(*CommandLine); // Note: Initializes AutoRTFM
 	GLog->SetCurrentThreadAsPrimaryThread();
 
diff --git a/Engine/Source/Runtime/Core/Public/HAL/LowLevelMemTrackerDefines.h b/Engine/Source/Runtime/Core/Public/HAL/LowLevelMemTrackerDefines.h
--- a/Engine/Source/Runtime/Core/Public/HAL/LowLevelMemTrackerDefines.h
+++ b/Engine/Source/Runtime/Core/Public/HAL/LowLevelMemTrackerDefines.h
@@ -11,6 +11,12 @@
 	#define ALLOW_LOW_LEVEL_MEM_TRACKER_IN_TEST 1	// enabled in Test by default to facilitate LLM captures in auto perf testing
 #endif
 
+// LLM is usually not enabled in programs, but some programs want it (e.g. unittests).
+// Define LLM_ALLOWED_IN_APPLICATION_TYPE in target.cs if needed in a program or other non-engine application type.
+#ifndef LLM_ALLOWED_IN_APPLICATION_TYPE
+	#define LLM_ALLOWED_IN_APPLICATION_TYPE WITH_ENGINE
+#endif
+
 // LLM_ENABLED_IN_CONFIG can be defined here or in the build environment only.
 // It cannot be defined in platform header files because it is included in c-language compilation units that
 // cannot include those headers.
@@ -19,6 +25,6 @@
 #ifndef LLM_ENABLED_IN_CONFIG 
 	#define LLM_ENABLED_IN_CONFIG ( \
 		!UE_BUILD_SHIPPING && (!UE_BUILD_TEST || ALLOW_LOW_LEVEL_MEM_TRACKER_IN_TEST) && \
-		WITH_ENGINE \
+		LLM_ALLOWED_IN_APPLICATION_TYPE \
 	)
 #endif
