[AutoRTFM] Add a whole lot of benchmarks for calling functions with different AutoRTFM modes applied

#Synced-CL 52547192

diff --git a/Engine/Source/Programs/AutoRTFMTests/Private/Benchmarks.cpp b/Engine/Source/Programs/AutoRTFMTests/Private/Benchmarks.cpp
--- a/Engine/Source/Programs/AutoRTFMTests/Private/Benchmarks.cpp
+++ b/Engine/Source/Programs/AutoRTFMTests/Private/Benchmarks.cpp
@@ -1,6 +1,7 @@
 // Copyright Epic Games, Inc. All Rights Reserved.
 
 #include "AutoRTFM/Testing.h"
+#include "Utils.h"
 #include "Catch2Includes.h"
 
 #include <vector>
@@ -815,3 +816,164 @@ TEST_CASE("Benchmarks.Write")
 
 #undef BENCH
 }
+
+// ---------------------------------------------------------------------------
+// Cross-mode inlining benchmarks
+//
+// Measures runtime impact of calling functions across AutoRTFM mode
+// boundaries. Each leaf function writes to heap memory (transaction-
+// instrumented) and uses a hash accumulator to prevent constant folding
+// and dead code elimination.
+//
+// Benchmark names use the call chain format: "Transact, Open, Enable"
+// means a transaction calls an open function which calls an enable function.
+// "(noinline)" marks functions forced noinline to measure call overhead.
+// ---------------------------------------------------------------------------
+
+namespace
+{
+
+constexpr unsigned BenchIterations = 1000000;
+
+// Knuth multiplicative hash -- cheap, non-trivial transform that prevents
+// the optimizer from constant-folding or strength-reducing the loop.
+inline unsigned HashStep(unsigned Value)
+{
+	return Value * 2654435761u;
+}
+
+struct BenchmarkHelpers
+{
+	// Leaf functions -- one per mode, with inlined and noinline variants.
+
+	AUTORTFM_OPEN static unsigned Open(unsigned* Out, unsigned Value)
+	{
+		*Out = Value;
+		return HashStep(Value);
+	}
+
+	AUTORTFM_OPEN [[clang::noinline]] static unsigned OpenNoInline(unsigned* Out, unsigned Value)
+	{
+		*Out = Value;
+		return HashStep(Value);
+	}
+
+	AUTORTFM_ENABLE static unsigned Enable(unsigned* Out, unsigned Value)
+	{
+		*Out = Value;
+		return HashStep(Value);
+	}
+
+	AUTORTFM_ENABLE [[clang::noinline]] static unsigned EnableNoInline(unsigned* Out, unsigned Value)
+	{
+		*Out = Value;
+		return HashStep(Value);
+	}
+
+	AUTORTFM_INTERNAL static unsigned Internal(unsigned* Out, unsigned Value)
+	{
+		*Out = Value;
+		return HashStep(Value);
+	}
+
+	AUTORTFM_INTERNAL [[clang::noinline]] static unsigned InternalNoInline(unsigned* Out, unsigned Value)
+	{
+		*Out = Value;
+		return HashStep(Value);
+	}
+
+	AUTORTFM_DISABLE static unsigned Disable(unsigned* Out, unsigned Value)
+	{
+		*Out = Value;
+		return HashStep(Value);
+	}
+
+	AUTORTFM_DISABLE [[clang::noinline]] static unsigned DisableNoInline(unsigned* Out, unsigned Value)
+	{
+		*Out = Value;
+		return HashStep(Value);
+	}
+
+	// Composite functions -- call chains across mode boundaries.
+
+	AUTORTFM_OPEN static unsigned OpenOpen(unsigned* Out, unsigned Value)
+	{
+		return Open(Out, Value);
+	}
+
+	AUTORTFM_OPEN static unsigned OpenEnable(unsigned* Out, unsigned Value)
+	{
+		return Enable(Out, Value);
+	}
+
+	AUTORTFM_OPEN static unsigned OpenDisable(unsigned* Out, unsigned Value)
+	{
+		return Disable(Out, Value);
+	}
+
+	AUTORTFM_OPEN [[clang::noinline]] static unsigned OpenDisableNoInline(unsigned* Out, unsigned Value)
+	{
+		return DisableNoInline(Out, Value);
+	}
+};
+
+} // anonymous namespace
+
+TEST_CASE("Benchmarks.CrossModeInlining")
+{
+	unsigned Sink = 0;
+	auto* SinkPtr = &Sink;
+
+#define BENCH_TRANSACT(Name, Func) \
+	BENCHMARK(Name) \
+	{ \
+		AutoRTFM::Transact([&] \
+		{ \
+			unsigned Acc = 1; \
+			for (unsigned I = 0; I < BenchIterations; ++I) \
+			{ \
+				Acc = BenchmarkHelpers::Func(SinkPtr, Acc); \
+			} \
+			*SinkPtr = Acc; \
+		}); \
+	}
+
+#define BENCH_OPEN(Name, Func) \
+	BENCHMARK(Name) \
+	{ \
+		unsigned Acc = 1; \
+		for (unsigned I = 0; I < BenchIterations; ++I) \
+		{ \
+			Acc = BenchmarkHelpers::Func(SinkPtr, Acc); \
+		} \
+		*SinkPtr = Acc; \
+	}
+
+	// Baseline.
+	BENCH_OPEN("Enable (baseline)",                    Enable);
+
+	// Single mode outside a transaction.
+	BENCH_OPEN("Open",                                 Open);
+	BENCH_OPEN("Open (noinline)",                      OpenNoInline);
+	BENCH_OPEN("Internal",                             Internal);
+	BENCH_OPEN("Internal (noinline)",                  InternalNoInline);
+	BENCH_OPEN("Disable",                              Disable);
+	BENCH_OPEN("Disable (noinline)",                   DisableNoInline);
+
+	// Single mode from a transaction.
+	BENCH_TRANSACT("Transact, Open",                   Open);
+	BENCH_TRANSACT("Transact, Open (noinline)",        OpenNoInline);
+	BENCH_TRANSACT("Transact, Enable",                 Enable);
+	BENCH_TRANSACT("Transact, Enable (noinline)",      EnableNoInline);
+	BENCH_TRANSACT("Transact, Internal",               Internal);
+	BENCH_TRANSACT("Transact, Internal (noinline)",    InternalNoInline);
+
+	// Call chains from a transaction.
+	BENCH_TRANSACT("Transact, Open, Open",             OpenOpen);
+	BENCH_TRANSACT("Transact, Open, Enable",           OpenEnable);
+	BENCH_TRANSACT("Transact, Open, Disable",          OpenDisable);
+	BENCH_TRANSACT("Transact, Open, Disable (noinline)", OpenDisableNoInline);
+
+#undef BENCH_TRANSACT
+#undef BENCH_OPEN
+}
