AutoRTFMTests: LLM: Undefine LLM in the AutoRTFMTests. We recently enabled it so we could test LLM, but it is intermittently crashing due to a mismatched push/pop in FLLMScope in AnsiMallocInternal. Disabling while we investigate.
#rnx
#Synced-CL 52610419

diff --git a/Engine/Source/Programs/AutoRTFMTests/AutoRTFMTests.Target.cs b/Engine/Source/Programs/AutoRTFMTests/AutoRTFMTests.Target.cs
--- a/Engine/Source/Programs/AutoRTFMTests/AutoRTFMTests.Target.cs
+++ b/Engine/Source/Programs/AutoRTFMTests/AutoRTFMTests.Target.cs
@@ -95,7 +95,7 @@ public class AutoRTFMTestsTarget : TargetRules
 		GlobalDefinitions.Add("USE_HITCH_DETECTION=1");
 
 		// Support unittests for LowLevelMemTracker
-		GlobalDefinitions.Add("LLM_ALLOWED_IN_APPLICATION_TYPE=1");
+		// GlobalDefinitions.Add("LLM_ALLOWED_IN_APPLICATION_TYPE=1");
 
 		// For testing insights, this sets `UE_TRACE_ENABLED`.
 		bEnableTrace = true;
